datasmith: Remove "E" suffix in versions

diff --git a/Engine/Source/Developer/Datasmith/DatasmithCore/Private/DatasmithUtils.cpp b/Engine/Source/Developer/Datasmith/DatasmithCore/Private/DatasmithUtils.cpp
--- a/Engine/Source/Developer/Datasmith/DatasmithCore/Private/DatasmithUtils.cpp
+++ b/Engine/Source/Developer/Datasmith/DatasmithCore/Private/DatasmithUtils.cpp
@@ -5,7 +5,6 @@
 #include "DatasmithMesh.h"
 #include "DatasmithSceneFactory.h"
 #include "IDatasmithSceneElements.h"
-#include "EnterpriseVersion.h"
 
 #include "RawMesh.h"
 #include "HAL/FileManager.h"
@@ -25,15 +24,15 @@ static const FString InvalidNames[] = {
 };
 
 template<typename T, int32 Size>
-static constexpr int GetArrLength(T(&)[Size]) { return Size; }
+static constexpr int GetArrayLength(T(&)[Size]) { return Size; }
 
 void FDatasmithUtils::SanitizeNameInplace(FString& InString)
 {
 	static const TCHAR Original[] = TEXT("ÀÁÂÃÄÅÆÇÈÉÊËÌÍÎÏÐÑÒÓÔÕÖ×ØÙÚÛÜÝÞßàáâãäåæçèéêëìíîïðñòóôõö÷øùúûüýþÿБбВвГгДдЁёЖжЗзИиЙйКкЛлМмНнОоПпРрСсТтУуФфХхЦцЧчШшЩщЪъЫыЬьЭэЮюЯя'\"");
 	static const TCHAR Modified[] = TEXT("AAAAAAECEEEEIIIIDNOOOOOx0UUUUYPsaaaaaaeceeeeiiiiOnoooood0uuuuypyBbVvGgDdEeJjZzIiYyKkLlMmNnOoPpRrSsTtUuFfJjTtCcSsSs__ii__EeYyYy__");
-	static_assert(GetArrLength(Original) == GetArrLength(Modified), "array size mismatch");
+	static_assert(GetArrayLength(Original) == GetArrayLength(Modified), "array size mismatch");
 
-	for (int32 i = 0; i < GetArrLength(Original); i++)
+	for (int32 i = 0; i < GetArrayLength(Original); i++)
 	{
 		InString.ReplaceCharInline(Original[i], Modified[i], ESearchCase::CaseSensitive);
 	}
@@ -124,29 +123,20 @@ FString FDatasmithUtils::SanitizeFileName(FString InString)
 	return InString;
 }
 
-float FDatasmithUtils::GetEnterpriseVersionAsFloat()
-{
-	const int32 EnterpriseVersion = GetEnterpriseVersionAsInt();
-	int32 VersionNumberOfDigits = 1;
-	for (int32 Version = EnterpriseVersion; Version /= 10; VersionNumberOfDigits++);
-
-	return ((float)EnterpriseVersion) / FMath::Pow(10, VersionNumberOfDigits - 1);
-}
-
 int32 FDatasmithUtils::GetEnterpriseVersionAsInt()
 {
 	const int32 MinorVersion = FEngineVersion::Current().GetMinor() * 10;
 	int32 MinorNumberOfDigits = 1;
 	for (int32 Version = MinorVersion; Version /= 10; MinorNumberOfDigits++);
-	
+
 	const int32 MajorVersion = FEngineVersion::Current().GetMajor() * FMath::Pow(10, MinorNumberOfDigits);
 
-	return MajorVersion + MinorVersion + ENTERPRISE_PATCH_VERSION;
+	return MajorVersion + MinorVersion;
 }
 
 FString FDatasmithUtils::GetEnterpriseVersionAsString()
 {
-	return FEngineVersion::Current().ToString( EVersionComponent::Minor ) + TEXT("E") + FString::FromInt( ENTERPRISE_PATCH_VERSION );
+	return FEngineVersion::Current().ToString( EVersionComponent::Minor );
 }
 
 float FDatasmithUtils::GetDatasmithFormatVersionAsFloat()
diff --git a/Engine/Source/Developer/Datasmith/DatasmithCore/Private/EnterpriseVersion.h b/Engine/Source/Developer/Datasmith/DatasmithCore/Private/EnterpriseVersion.h
deleted file mode 100644
diff --git a/Engine/Source/Developer/Datasmith/DatasmithCore/Public/DatasmithUtils.h b/Engine/Source/Developer/Datasmith/DatasmithCore/Public/DatasmithUtils.h
--- a/Engine/Source/Developer/Datasmith/DatasmithCore/Public/DatasmithUtils.h
+++ b/Engine/Source/Developer/Datasmith/DatasmithCore/Public/DatasmithUtils.h
@@ -19,7 +19,6 @@ public:
 	static FString SanitizeObjectName(FString InString);
 	static FString SanitizeFileName(FString InString);
 
-	static float GetEnterpriseVersionAsFloat();
 	static int32 GetEnterpriseVersionAsInt();
 	static FString GetEnterpriseVersionAsString();
 
