Datasmith: Fix missing patch number in exporter displayed version.

diff --git a/Engine/Source/Developer/Datasmith/DatasmithCore/Private/DatasmithUtils.cpp b/Engine/Source/Developer/Datasmith/DatasmithCore/Private/DatasmithUtils.cpp
--- a/Engine/Source/Developer/Datasmith/DatasmithCore/Private/DatasmithUtils.cpp
+++ b/Engine/Source/Developer/Datasmith/DatasmithCore/Private/DatasmithUtils.cpp
@@ -125,18 +125,19 @@ FString FDatasmithUtils::SanitizeFileName(FString InString)
 
 int32 FDatasmithUtils::GetEnterpriseVersionAsInt()
 {
+	const int32 PatchVersion = FEngineVersion::Current().GetPatch();
 	const int32 MinorVersion = FEngineVersion::Current().GetMinor() * 10;
 	int32 MinorNumberOfDigits = 1;
 	for (int32 Version = MinorVersion; Version /= 10; MinorNumberOfDigits++);
 
 	const int32 MajorVersion = FEngineVersion::Current().GetMajor() * FMath::Pow(10, MinorNumberOfDigits);
 
-	return MajorVersion + MinorVersion;
+	return MajorVersion + MinorVersion + PatchVersion;
 }
 
 FString FDatasmithUtils::GetEnterpriseVersionAsString()
 {
-	return FEngineVersion::Current().ToString( EVersionComponent::Minor );
+	return FEngineVersion::Current().ToString( EVersionComponent::Patch );
 }
 
 float FDatasmithUtils::GetDatasmithFormatVersionAsFloat()
