Illegal characters in texture name don't survive datasmith export

Integrated from Dev-Enterprise (CL 11279033)

diff --git a/Engine/Source/Developer/Datasmith/DatasmithFacade/Private/DatasmithFacadeMaterial.cpp b/Engine/Source/Developer/Datasmith/DatasmithFacade/Private/DatasmithFacadeMaterial.cpp
--- a/Engine/Source/Developer/Datasmith/DatasmithFacade/Private/DatasmithFacadeMaterial.cpp
+++ b/Engine/Source/Developer/Datasmith/DatasmithFacade/Private/DatasmithFacadeMaterial.cpp
@@ -66,7 +66,7 @@ void FDatasmithFacadeMaterial::AddTexture(
 	if (!FString(InTextureFilePath).IsEmpty())
 	{
 		// Make Datasmith texture name from file name.
-		FString FileName = FDatasmithUtils::SanitizeFileName(FPaths::GetBaseFilename(InTextureFilePath));
+		FString FileName = FDatasmithUtils::SanitizeObjectName(FPaths::GetBaseFilename(InTextureFilePath));
 		FString TextureName = FString::Printf(TEXT("%ls_%d"), *FileName, int(InTextureMode));
 
 		// Create a new Datasmith material texture property.
