Integrating CL#11259478 from Dev-Enterprise
Fixing memory leak when GC is not called in datasmith mesh exporter

JeanMichel.Dignard

diff --git a/Engine/Source/Developer/Datasmith/DatasmithExporter/Private/DatasmithMeshExporter.cpp b/Engine/Source/Developer/Datasmith/DatasmithExporter/Private/DatasmithMeshExporter.cpp
--- a/Engine/Source/Developer/Datasmith/DatasmithExporter/Private/DatasmithMeshExporter.cpp
+++ b/Engine/Source/Developer/Datasmith/DatasmithExporter/Private/DatasmithMeshExporter.cpp
@@ -7,22 +7,28 @@
 #include "DatasmithSceneFactory.h"
 #include "DatasmithUtils.h"
 
+#include "Containers/LockFreeList.h"
 #include "HAL/FileManager.h"
 #include "HAL/PlatformFilemanager.h"
 #include "Misc/Paths.h"
 #include "Serialization/MemoryWriter.h"
+#include "StaticMeshAttributes.h"
 #include "StaticMeshOperations.h"
 #include "UVMapSettings.h"
-#include "StaticMeshAttributes.h"
-
 
-namespace
+/**
+ *	Implementation class of the DatasmithMeshExporter
+ *	We use a lockfree UDatasmithMesh pool to avoid creating new UObject when exporting and reduces our memory footprint.
+ */
+class FDatasmithMeshExporterImpl
 {
-	UDatasmithMesh* FDatasmithMeshToUDatasmithMesh( const FDatasmithMesh& Mesh, bool bValidateRawMesh )
-	{
-		UDatasmithMesh* UMesh = NewObject< UDatasmithMesh >();
-		UMesh->AddToRoot();
+public:
 
+	/**
+	 * This function allows reusing an instanced UDatasmithMesh. Reusing the same object will avoid creating new garbage in memory.
+	 */
+	void FillUDatasmithMeshFromFDatasmithMesh(UDatasmithMesh* UMesh, const FDatasmithMesh& Mesh, bool bValidateRawMesh )
+	{
 		UMesh->MeshName = Mesh.GetName();
 
 		FRawMesh RawMesh;
@@ -45,9 +51,38 @@ namespace
 
 			UMesh->SourceModels.Add( LODModel );
 		}
-
-		return UMesh;
 	}
+
+	void PreExport(FDatasmithMesh& DatasmithMesh, const TCHAR* Filepath, const TCHAR* Filename, EDSExportLightmapUV LightmapUV);
+	void PostExport(const FDatasmithMesh& DatasmithMesh, TSharedRef< IDatasmithMeshElement > MeshElement);
+
+	void CreateDefaultUVs(FDatasmithMesh& DatasmithMesh);
+	void RegisterStaticMeshAttributes(FMeshDescription& MeshDescription);
+
+	UDatasmithMesh* GetPooledUDatasmithMesh();
+	void ReturnUDatasmithMeshToPool(UDatasmithMesh*& UMesh);
+	void ClearUDatasmithMeshPool();
+
+	FString LastError;
+
+private:
+	//A pool of UDatasmithMesh that we use to avoid creating new UObject, this greatly reduce the garbage created when one instance of FDatasmithMeshExporter is used to export multiple Meshes.
+	TLockFreePointerListLIFO< UDatasmithMesh > DatasmithMeshUObjectPool;
+	static TAtomic<int32> NumberOfUMeshPendingGC;
+};
+
+TAtomic<int32> FDatasmithMeshExporterImpl::NumberOfUMeshPendingGC(0);
+
+FDatasmithMeshExporter::FDatasmithMeshExporter()
+{
+	Impl = new FDatasmithMeshExporterImpl();
+}
+
+FDatasmithMeshExporter::~FDatasmithMeshExporter()
+{
+	Impl->ClearUDatasmithMeshPool();
+	delete Impl;
+	Impl = nullptr;
 }
 
 TSharedPtr< IDatasmithMeshElement > FDatasmithMeshExporter::ExportToUObject( const TCHAR* Filepath, const TCHAR* Filename, FDatasmithMesh& Mesh, FDatasmithMesh* CollisionMesh, EDSExportLightmapUV LightmapUV )
@@ -58,17 +93,20 @@ TSharedPtr< IDatasmithMeshElement > FDatasmithMeshExporter::ExportToUObject( con
 	FString NormalizedFilename = Filename;
 	FPaths::NormalizeFilename( NormalizedFilename );
 
-	PreExport( Mesh, *NormalizedFilepath, *NormalizedFilename, LightmapUV );
+	Impl->PreExport( Mesh, *NormalizedFilepath, *NormalizedFilename, LightmapUV );
 
-	TArray< UDatasmithMesh* > MeshesToExport;
+	TArray< UDatasmithMesh*, TInlineAllocator<2>> MeshesToExport;
 
-	// Static mesh
-	MeshesToExport.Add( FDatasmithMeshToUDatasmithMesh( Mesh, true ) );
+	// Static mesh, we keep a static UDatasmithMesh alive as a utility object and re-use it for every export instead of creating a new one every time. This avoid creating garbage in memory.
+	UDatasmithMesh* UMesh = Impl->GetPooledUDatasmithMesh();
+	Impl->FillUDatasmithMeshFromFDatasmithMesh(UMesh, Mesh, true);
+	MeshesToExport.Add(UMesh);
 
 	// Collision mesh
 	if ( CollisionMesh )
 	{
-		UDatasmithMesh* DSColMesh = FDatasmithMeshToUDatasmithMesh( *CollisionMesh, false );
+		UDatasmithMesh* DSColMesh = Impl->GetPooledUDatasmithMesh();
+		Impl->FillUDatasmithMeshFromFDatasmithMesh( DSColMesh, *CollisionMesh, false );
 		DSColMesh->bIsCollisionMesh = true;
 
 		MeshesToExport.Add( DSColMesh );
@@ -81,7 +119,11 @@ TSharedPtr< IDatasmithMeshElement > FDatasmithMeshExporter::ExportToUObject( con
 
 	if ( !Archive.IsValid() )
 	{
-		LastError = FString::Printf( TEXT("Failed writing to file %s"), *FullPath );
+		Impl->LastError = FString::Printf( TEXT("Failed writing to file %s"), *FullPath );
+		for (UDatasmithMesh*& MeshToExport : MeshesToExport)
+		{
+			Impl->ReturnUDatasmithMeshToPool(MeshToExport);
+		}
 
 		return TSharedPtr< IDatasmithMeshElement >();
 	}
@@ -91,7 +133,7 @@ TSharedPtr< IDatasmithMeshElement > FDatasmithMeshExporter::ExportToUObject( con
 	*Archive << NumMeshes;
 
 	FMD5 MD5;
-	for ( UDatasmithMesh* MeshToExport : MeshesToExport )
+	for ( UDatasmithMesh*& MeshToExport : MeshesToExport )
 	{
 		TArray< uint8 > Bytes;
 		FMemoryWriter MemoryWriter( Bytes, true );
@@ -111,7 +153,8 @@ TSharedPtr< IDatasmithMeshElement > FDatasmithMeshExporter::ExportToUObject( con
 		}
 
 		*Archive << Bytes;
-		MeshToExport->RemoveFromRoot();
+		//Return the UDatasmithMesh to the pool.
+		Impl->ReturnUDatasmithMeshToPool( MeshToExport );
 	}
 	FMD5Hash Hash;
 	Hash.Set(MD5);
@@ -122,12 +165,17 @@ TSharedPtr< IDatasmithMeshElement > FDatasmithMeshExporter::ExportToUObject( con
 	MeshElement->SetFile( *FullPath );
 	MeshElement->SetFileHash(Hash);
 
-	PostExport( Mesh, MeshElement.ToSharedRef() );
+	Impl->PostExport( Mesh, MeshElement.ToSharedRef() );
 
 	return MeshElement;
 }
 
-void FDatasmithMeshExporter::PreExport( FDatasmithMesh& Mesh, const TCHAR* Filepath, const TCHAR* Filename, EDSExportLightmapUV LightmapUV )
+FString FDatasmithMeshExporter::GetLastError() const
+{
+	return Impl->LastError;
+}
+
+void FDatasmithMeshExporterImpl::PreExport( FDatasmithMesh& Mesh, const TCHAR* Filepath, const TCHAR* Filename, EDSExportLightmapUV LightmapUV )
 {
 	// If the mesh doesn't have a name, use the filename as its name
 	if ( FCString::Strlen( Mesh.GetName() ) == 0 )
@@ -147,7 +195,7 @@ void FDatasmithMeshExporter::PreExport( FDatasmithMesh& Mesh, const TCHAR* Filep
 	}
 }
 
-void FDatasmithMeshExporter::PostExport( const FDatasmithMesh& DatasmithMesh, TSharedRef< IDatasmithMeshElement > MeshElement )
+void FDatasmithMeshExporterImpl::PostExport( const FDatasmithMesh& DatasmithMesh, TSharedRef< IDatasmithMeshElement > MeshElement )
 {
 	FBox Extents = DatasmithMesh.GetExtents();
 	float Width = Extents.Max[0] - Extents.Min[0];
@@ -159,7 +207,7 @@ void FDatasmithMeshExporter::PostExport( const FDatasmithMesh& DatasmithMesh, TS
 	MeshElement->SetLightmapSourceUV( DatasmithMesh.GetLightmapSourceUVChannel() );
 }
 
-void FDatasmithMeshExporter::CreateDefaultUVs( FDatasmithMesh& Mesh )
+void FDatasmithMeshExporterImpl::CreateDefaultUVs( FDatasmithMesh& Mesh )
 {
 	if ( Mesh.GetUVChannelsCount() > 0 )
 	{
@@ -207,7 +255,7 @@ void FDatasmithMeshExporter::CreateDefaultUVs( FDatasmithMesh& Mesh )
 	}
 }
 
-void FDatasmithMeshExporter::RegisterStaticMeshAttributes(FMeshDescription& MeshDescription)
+void FDatasmithMeshExporterImpl::RegisterStaticMeshAttributes(FMeshDescription& MeshDescription)
 {
 	// This is a local inlining of the UStaticMesh::RegisterMeshAttributes() function, to avoid having to include "Engine" in our dependencies..
 	///////
@@ -229,3 +277,48 @@ void FDatasmithMeshExporter::RegisterStaticMeshAttributes(FMeshDescription& Mesh
 	// Add basic polygon group attributes
 	MeshDescription.PolygonGroupAttributes().RegisterAttribute<FName>(MeshAttribute::PolygonGroup::ImportedMaterialSlotName); //The unique key to match the mesh material slot
 }
+
+UDatasmithMesh* FDatasmithMeshExporterImpl::GetPooledUDatasmithMesh()
+{
+	while (UDatasmithMesh* PooledMesh = DatasmithMeshUObjectPool.Pop())
+	{
+		PooledMesh->AddToRoot();
+		return PooledMesh;
+	}
+
+	return NewObject< UDatasmithMesh >((UObject*)GetTransientPackage(), TEXT("DatasmithExporter_TransientPooledUDatasmithMesh"), RF_Transient | RF_MarkAsRootSet);
+}
+
+void FDatasmithMeshExporterImpl::ReturnUDatasmithMeshToPool(UDatasmithMesh*& UMesh)
+{
+	//Clear the UDatasmithMesh.
+	UMesh->SourceModels.Empty();
+	UMesh->bIsCollisionMesh = false;
+
+	//Put it back into the pool
+	DatasmithMeshUObjectPool.Push(UMesh);
+	//Null the reference to make sure we don't use the object returned to the pool.
+	UMesh = nullptr;
+}
+
+void FDatasmithMeshExporterImpl::ClearUDatasmithMeshPool()
+{
+	TArray<UDatasmithMesh*> PooledMeshes;
+	DatasmithMeshUObjectPool.PopAll(PooledMeshes);
+
+	for (UDatasmithMesh* UMesh : PooledMeshes)
+	{
+		UMesh->RemoveFromRoot();
+		UMesh->MarkPendingKill();
+	}
+
+	//Keep track of the number of garbage UObject generated by clearing the cache so that we can trigger the GC after a while.
+	//Even if our UObjects are basically empty at this point and don't have a big memory footprint, the engine will assert when reaching around 65k UObjects.
+	//So we need to call the GC before that point.
+	NumberOfUMeshPendingGC += PooledMeshes.Num();
+	if (NumberOfUMeshPendingGC % 2000 == 0)
+	{
+		CollectGarbage(RF_NoFlags);
+		NumberOfUMeshPendingGC = 0;
+	}
+}
diff --git a/Engine/Source/Developer/Datasmith/DatasmithExporter/Public/DatasmithMeshExporter.h b/Engine/Source/Developer/Datasmith/DatasmithExporter/Public/DatasmithMeshExporter.h
--- a/Engine/Source/Developer/Datasmith/DatasmithExporter/Public/DatasmithMeshExporter.h
+++ b/Engine/Source/Developer/Datasmith/DatasmithExporter/Public/DatasmithMeshExporter.h
@@ -7,11 +7,13 @@
 class FDatasmithMesh;
 class FDatasmithMeshExporterImpl;
 class IDatasmithMeshElement;
-struct FMeshDescription;
 
 class DATASMITHEXPORTER_API FDatasmithMeshExporter
 {
 public:
+	FDatasmithMeshExporter();
+	virtual ~FDatasmithMeshExporter();
+
 	/**
 	 * Exports the DatasmithMesh as a UObject.
 	 *
@@ -28,14 +30,9 @@ public:
 	/**
 	 * @return The error that happened during the last export, if any
 	 */
-	FString GetLastError() const { return LastError; }
+	FString GetLastError() const;
 
 private:
-	void PreExport( FDatasmithMesh& DatasmithMesh, const TCHAR* Filepath, const TCHAR* Filename, EDSExportLightmapUV LightmapUV );
-	void PostExport( const FDatasmithMesh& DatasmithMesh, TSharedRef< IDatasmithMeshElement > MeshElement );
-
-	void CreateDefaultUVs( FDatasmithMesh& DatasmithMesh );
-	void RegisterStaticMeshAttributes( FMeshDescription& MeshDescription);
 
-	FString LastError;
+	FDatasmithMeshExporterImpl* Impl;
 };
