Integrating CL# 12246665 from Dev-Enterprise
Fix datasmith export collision mesh.

diff --git a/Engine/Source/Developer/Datasmith/DatasmithExporter/Private/DatasmithMeshExporter.cpp b/Engine/Source/Developer/Datasmith/DatasmithExporter/Private/DatasmithMeshExporter.cpp
--- a/Engine/Source/Developer/Datasmith/DatasmithExporter/Private/DatasmithMeshExporter.cpp
+++ b/Engine/Source/Developer/Datasmith/DatasmithExporter/Private/DatasmithMeshExporter.cpp
@@ -10,6 +10,7 @@
 #include "Containers/LockFreeList.h"
 #include "HAL/FileManager.h"
 #include "HAL/PlatformFilemanager.h"
+#include "Misc/Guid.h"
 #include "Misc/Paths.h"
 #include "Serialization/MemoryWriter.h"
 #include "StaticMeshAttributes.h"
@@ -23,52 +24,69 @@
 class FDatasmithMeshExporterImpl
 {
 public:
+	FDatasmithMeshExporterImpl() :
+		UniqueID(FGuid::NewGuid())
+	{}
 
-	/**
-	 * This function allows reusing an instanced UDatasmithMesh. Reusing the same object will avoid creating new garbage in memory.
-	 */
-	void FillUDatasmithMeshFromFDatasmithMesh(UDatasmithMesh* UMesh, const FDatasmithMesh& Mesh, bool bValidateRawMesh )
+	~FDatasmithMeshExporterImpl()
 	{
-		UMesh->MeshName = Mesh.GetName();
-
-		FRawMesh RawMesh;
-		FDatasmithMeshUtils::ToRawMesh(Mesh, RawMesh, bValidateRawMesh);
-
-		FDatasmithMeshSourceModel BaseModel;
-		BaseModel.RawMeshBulkData.SaveRawMesh( RawMesh );
-
-		UMesh->SourceModels.Add( BaseModel );
-
-		for ( int32 LODIndex = 0; LODIndex < Mesh.GetLODsCount(); ++LODIndex )
-		{
-			FDatasmithMeshUtils::ToRawMesh(Mesh.GetLOD( LODIndex ), RawMesh, bValidateRawMesh);
-
-			FRawMeshBulkData LODRawMeshBulkData;
-			LODRawMeshBulkData.SaveRawMesh( RawMesh );
-
-			FDatasmithMeshSourceModel LODModel;
-			LODModel.RawMeshBulkData = LODRawMeshBulkData;
-
-			UMesh->SourceModels.Add( LODModel );
-		}
+		ClearUDatasmithMeshPool();
 	}
 
+	/**
+	 * Generate and fill a UDatasmithMesh from a FDatasmithMesh.
+	 * Note: The returned TSharedPtr uses a special destructor to return the UDatasmithMesh to a pool instead of destroying it (which wouldn't work anyways).				
+	 */
+	TSharedPtr<UDatasmithMesh> GeneratePooledUDatasmithMeshFromFDatasmithMesh( const FDatasmithMesh& Mesh, bool bIsCollisionMesh );
+
 	void PreExport(FDatasmithMesh& DatasmithMesh, const TCHAR* Filepath, const TCHAR* Filename, EDSExportLightmapUV LightmapUV);
 	void PostExport(const FDatasmithMesh& DatasmithMesh, TSharedRef< IDatasmithMeshElement > MeshElement);
 
 	void CreateDefaultUVs(FDatasmithMesh& DatasmithMesh);
 	void RegisterStaticMeshAttributes(FMeshDescription& MeshDescription);
 
-	UDatasmithMesh* GetPooledUDatasmithMesh();
-	void ReturnUDatasmithMeshToPool(UDatasmithMesh*& UMesh);
-	void ClearUDatasmithMeshPool();
-
 	FString LastError;
 
 private:
-	//A pool of UDatasmithMesh that we use to avoid creating new UObject, this greatly reduce the garbage created when one instance of FDatasmithMeshExporter is used to export multiple Meshes.
+	/**
+	 * This function allows reusing an instanced UDatasmithMesh. Reusing the same object will avoid creating new garbage in memory.
+	 */
+	void FillUDatasmithMeshFromFDatasmithMesh( TSharedPtr<UDatasmithMesh>& UMesh, const FDatasmithMesh& Mesh, bool bValidateRawMesh );
+
+	TSharedPtr<UDatasmithMesh> GetPooledUDatasmithMesh(bool bIsCollisionMesh);
+
+	void ReturnUDatasmithMeshToPool(UDatasmithMesh*& UMesh);
+
+	void ClearUDatasmithMeshPool();
+
+	/**
+	 * A pool of UDatasmithMesh that we use to avoid creating new UObject, this greatly reduce the garbage created when one instance of FDatasmithMeshExporter is used to export multiple Meshes.
+	 */
 	TLockFreePointerListLIFO< UDatasmithMesh > DatasmithMeshUObjectPool;
 	static TAtomic<int32> NumberOfUMeshPendingGC;
+
+	/**
+	 * This exporter's UniqueID. Used to make sure pooled meshes are not using the same names across different threads.
+	 */
+	FGuid UniqueID;
+
+	/**
+	 * Custom deleter used by the TSharedPtr of a pooled UDatasmithMesh to return it to the pool instead of deleting them.
+	 */
+	struct FPooledUDatasmithMeshDeleter
+	{
+		FPooledUDatasmithMeshDeleter(FDatasmithMeshExporterImpl* InExporterPtr)
+			: ExporterPtr(InExporterPtr)
+		{}
+
+		FORCEINLINE void operator()(UDatasmithMesh* Object) const
+		{
+			check(Object != nullptr && ExporterPtr != nullptr);
+			ExporterPtr->ReturnUDatasmithMeshToPool(Object);
+		}
+
+		FDatasmithMeshExporterImpl* ExporterPtr;
+	};
 };
 
 TAtomic<int32> FDatasmithMeshExporterImpl::NumberOfUMeshPendingGC(0);
@@ -80,7 +98,6 @@ FDatasmithMeshExporter::FDatasmithMeshExporter()
 
 FDatasmithMeshExporter::~FDatasmithMeshExporter()
 {
-	Impl->ClearUDatasmithMeshPool();
 	delete Impl;
 	Impl = nullptr;
 }
@@ -95,21 +112,17 @@ TSharedPtr< IDatasmithMeshElement > FDatasmithMeshExporter::ExportToUObject( con
 
 	Impl->PreExport( Mesh, *NormalizedFilepath, *NormalizedFilename, LightmapUV );
 
-	TArray< UDatasmithMesh*, TInlineAllocator<2>> MeshesToExport;
+	TArray<TSharedPtr<UDatasmithMesh>, TInlineAllocator<2>> MeshesToExport;
 
 	// Static mesh, we keep a static UDatasmithMesh alive as a utility object and re-use it for every export instead of creating a new one every time. This avoid creating garbage in memory.
-	UDatasmithMesh* UMesh = Impl->GetPooledUDatasmithMesh();
-	Impl->FillUDatasmithMeshFromFDatasmithMesh(UMesh, Mesh, true);
-	MeshesToExport.Add(UMesh);
+	bool bIsCollisionMesh = false;
+	MeshesToExport.Add( Impl->GeneratePooledUDatasmithMeshFromFDatasmithMesh( Mesh, bIsCollisionMesh ) );
 
 	// Collision mesh
 	if ( CollisionMesh )
 	{
-		UDatasmithMesh* DSColMesh = Impl->GetPooledUDatasmithMesh();
-		Impl->FillUDatasmithMeshFromFDatasmithMesh( DSColMesh, *CollisionMesh, false );
-		DSColMesh->bIsCollisionMesh = true;
-
-		MeshesToExport.Add( DSColMesh );
+		bIsCollisionMesh = true;
+		MeshesToExport.Add( Impl->GeneratePooledUDatasmithMeshFromFDatasmithMesh( *CollisionMesh, bIsCollisionMesh ) );
 	}
 
 	FString FullPath = FPaths::Combine( *NormalizedFilepath, FPaths::SetExtension( *NormalizedFilename, UDatasmithMesh::GetFileExtension() ) );
@@ -120,10 +133,6 @@ TSharedPtr< IDatasmithMeshElement > FDatasmithMeshExporter::ExportToUObject( con
 	if ( !Archive.IsValid() )
 	{
 		Impl->LastError = FString::Printf( TEXT("Failed writing to file %s"), *FullPath );
-		for (UDatasmithMesh*& MeshToExport : MeshesToExport)
-		{
-			Impl->ReturnUDatasmithMeshToPool(MeshToExport);
-		}
 
 		return TSharedPtr< IDatasmithMeshElement >();
 	}
@@ -133,7 +142,7 @@ TSharedPtr< IDatasmithMeshElement > FDatasmithMeshExporter::ExportToUObject( con
 	*Archive << NumMeshes;
 
 	FMD5 MD5;
-	for ( UDatasmithMesh*& MeshToExport : MeshesToExport )
+	for ( TSharedPtr<UDatasmithMesh>& MeshToExport : MeshesToExport )
 	{
 		TArray< uint8 > Bytes;
 		FMemoryWriter MemoryWriter( Bytes, true );
@@ -153,8 +162,6 @@ TSharedPtr< IDatasmithMeshElement > FDatasmithMeshExporter::ExportToUObject( con
 		}
 
 		*Archive << Bytes;
-		//Return the UDatasmithMesh to the pool.
-		Impl->ReturnUDatasmithMeshToPool( MeshToExport );
 	}
 	FMD5Hash Hash;
 	Hash.Set(MD5);
@@ -278,22 +285,62 @@ void FDatasmithMeshExporterImpl::RegisterStaticMeshAttributes(FMeshDescription&
 	MeshDescription.PolygonGroupAttributes().RegisterAttribute<FName>(MeshAttribute::PolygonGroup::ImportedMaterialSlotName); //The unique key to match the mesh material slot
 }
 
-UDatasmithMesh* FDatasmithMeshExporterImpl::GetPooledUDatasmithMesh()
+TSharedPtr<UDatasmithMesh> FDatasmithMeshExporterImpl::GeneratePooledUDatasmithMeshFromFDatasmithMesh(const FDatasmithMesh& Mesh, bool bIsCollisionMesh)
+{
+	TSharedPtr<UDatasmithMesh> PooledMesh = GetPooledUDatasmithMesh(bIsCollisionMesh);
+	FillUDatasmithMeshFromFDatasmithMesh(PooledMesh, Mesh, !bIsCollisionMesh);
+	PooledMesh->bIsCollisionMesh = bIsCollisionMesh;
+	return PooledMesh;
+}
+
+/**
+ * This function allows reusing an instanced UDatasmithMesh. Reusing the same object will avoid creating new garbage in memory.
+ */
+void FDatasmithMeshExporterImpl::FillUDatasmithMeshFromFDatasmithMesh(TSharedPtr<UDatasmithMesh>& UMesh, const FDatasmithMesh& Mesh, bool bValidateRawMesh)
 {
-	while (UDatasmithMesh* PooledMesh = DatasmithMeshUObjectPool.Pop())
+	UMesh->MeshName = Mesh.GetName();
+
+	FRawMesh RawMesh;
+	FDatasmithMeshUtils::ToRawMesh(Mesh, RawMesh, bValidateRawMesh);
+
+	FDatasmithMeshSourceModel BaseModel;
+	BaseModel.RawMeshBulkData.SaveRawMesh(RawMesh);
+
+	UMesh->SourceModels.Add(BaseModel);
+
+	for (int32 LODIndex = 0; LODIndex < Mesh.GetLODsCount(); ++LODIndex)
+	{
+		FDatasmithMeshUtils::ToRawMesh(Mesh.GetLOD(LODIndex), RawMesh, bValidateRawMesh);
+
+		FRawMeshBulkData LODRawMeshBulkData;
+		LODRawMeshBulkData.SaveRawMesh(RawMesh);
+
+		FDatasmithMeshSourceModel LODModel;
+		LODModel.RawMeshBulkData = LODRawMeshBulkData;
+
+		UMesh->SourceModels.Add(LODModel);
+	}
+}
+
+TSharedPtr<UDatasmithMesh> FDatasmithMeshExporterImpl::GetPooledUDatasmithMesh(bool bIsCollisionMesh)
+{
+	const FString GuidString = UniqueID.ToString();
+	const FString CollisionString = bIsCollisionMesh ? TEXT("_Collision") : TEXT("");
+	const FString PooledMeshName = FString::Printf(TEXT("DatasmithExporter_%s_TransientPooledUDatasmithMesh%s"), *GuidString, *CollisionString);
+
+	if (UDatasmithMesh* PooledMesh = DatasmithMeshUObjectPool.Pop())
 	{
-		PooledMesh->AddToRoot();
-		return PooledMesh;
+		PooledMesh->Rename(*PooledMeshName);
+		return TSharedPtr<UDatasmithMesh>(PooledMesh, FPooledUDatasmithMeshDeleter(this));
 	}
 
-	return NewObject< UDatasmithMesh >((UObject*)GetTransientPackage(), TEXT("DatasmithExporter_TransientPooledUDatasmithMesh"), RF_Transient | RF_MarkAsRootSet);
+	return TSharedPtr<UDatasmithMesh>(NewObject< UDatasmithMesh >((UObject*)GetTransientPackage(), *PooledMeshName, RF_Transient | RF_MarkAsRootSet), FPooledUDatasmithMeshDeleter(this));
 }
 
 void FDatasmithMeshExporterImpl::ReturnUDatasmithMeshToPool(UDatasmithMesh*& UMesh)
 {
 	//Clear the UDatasmithMesh.
 	UMesh->SourceModels.Empty();
-	UMesh->bIsCollisionMesh = false;
 
 	//Put it back into the pool
 	DatasmithMeshUObjectPool.Push(UMesh);
