There should be a way of resetting cache directory in the revit datasmith connections window

diff --git a/Engine/Source/Developer/Datasmith/DatasmithExporterUI/Private/DirectLinkUI.cpp b/Engine/Source/Developer/Datasmith/DatasmithExporterUI/Private/DirectLinkUI.cpp
--- a/Engine/Source/Developer/Datasmith/DatasmithExporterUI/Private/DirectLinkUI.cpp
+++ b/Engine/Source/Developer/Datasmith/DatasmithExporterUI/Private/DirectLinkUI.cpp
@@ -32,9 +32,9 @@ namespace DirectLinkUIUtils
 		return FPaths::Combine( FPaths::GeneratedConfigDir(), TEXT("DirectLinkExporter") ).Append( TEXT(".ini") );;
 	}
 
-	FString GetConfigCacheDirectorySectionAndValue()
+	FString GetConfigCacheDirectorySectionAndValue( bool bInDefaultCacheDir )
 	{
-		return TEXT("DLCacheFolder");
+		return bInDefaultCacheDir ? TEXT("DLDefaultCacheFolder") : TEXT("DLCacheFolder");
 	}
 
 	EHorizontalAlignment GetWindowTitleAlignement()
@@ -61,12 +61,19 @@ namespace DirectLinkUIUtils
 FDirectLinkUI::FDirectLinkUI()
 {
 	FString ConfigPath = DirectLinkUIUtils::GetConfigPath();
-	FString DirectLinkCacheSectionAndValue = DirectLinkUIUtils::GetConfigCacheDirectorySectionAndValue();
+	FString DirectLinkCacheSectionAndValue = DirectLinkUIUtils::GetConfigCacheDirectorySectionAndValue( false );
+	FString DirectLinkDefaultCacheSectionAndValue = DirectLinkUIUtils::GetConfigCacheDirectorySectionAndValue( true );
 
 	FScopeLock Lock( &CriticalSectionCacheDirectory );
+	if ( !GConfig->GetString( *DirectLinkDefaultCacheSectionAndValue, *DirectLinkDefaultCacheSectionAndValue, DefaultDirectLinkCacheDirectory, ConfigPath ) )
+	{
+		DefaultDirectLinkCacheDirectory = FPaths::Combine( FPlatformProcess::UserTempDir(), TEXT("DLExporter") );
+		SaveCacheDirectory( DefaultDirectLinkCacheDirectory, true );
+	}
 	if ( !GConfig->GetString( *DirectLinkCacheSectionAndValue, *DirectLinkCacheSectionAndValue, DirectLinkCacheDirectory, ConfigPath ) )
 	{
-		DirectLinkCacheDirectory = FPaths::Combine( FPlatformProcess::UserTempDir(), TEXT("DLExporter") );
+		DirectLinkCacheDirectory = DefaultDirectLinkCacheDirectory;
+		SaveCacheDirectory( DirectLinkCacheDirectory, false );
 	}
 }
 
@@ -121,6 +128,7 @@ void FDirectLinkUI::OpenDirectLinkStreamWindow()
 							.DefaultCacheDirectory( DirectLinkCacheDirectory )
 							// The slate application should always be close before this module.
 							.OnCacheDirectoryChanged_Raw( this, &FDirectLinkUI::OnCacheDirectoryChanged )
+							.OnCacheDirectoryReset_Raw( this, &FDirectLinkUI::OnCacheDirectoryReset )
 						]
 					]
 				);
@@ -159,11 +167,29 @@ void FDirectLinkUI::OnCacheDirectoryChanged(const FString& InNewCacheDirectory)
 		DirectLinkCacheDirectory = InNewCacheDirectory;
 	}
 
+	SaveCacheDirectory( InNewCacheDirectory, false );
+
+}
+
+FString FDirectLinkUI::OnCacheDirectoryReset()
+{
+	{
+		FScopeLock ReadLock( &CriticalSectionCacheDirectory );
+		DirectLinkCacheDirectory = DefaultDirectLinkCacheDirectory;
+	}
+
+	SaveCacheDirectory( DefaultDirectLinkCacheDirectory, false );
+
+	return DefaultDirectLinkCacheDirectory;
+}
+
+void FDirectLinkUI::SaveCacheDirectory( const FString& InCacheDir, bool bInDefaultCacheDir )
+{
 	FString ConfigPath = DirectLinkUIUtils::GetConfigPath();
-	FString DirectLinkCacheSectionAndValue = DirectLinkUIUtils::GetConfigCacheDirectorySectionAndValue();
+	FString DirectLinkCacheSectionAndValue = DirectLinkUIUtils::GetConfigCacheDirectorySectionAndValue( bInDefaultCacheDir );
 
 	// Save to config file
-	GConfig->SetString( *DirectLinkCacheSectionAndValue, *DirectLinkCacheSectionAndValue, *InNewCacheDirectory, ConfigPath );
+	GConfig->SetString( *DirectLinkCacheSectionAndValue, *DirectLinkCacheSectionAndValue, *InCacheDir, ConfigPath );
 }
 
 #undef LOCTEXT_NAMESPACE
diff --git a/Engine/Source/Developer/Datasmith/DatasmithExporterUI/Private/DirectLinkUI.h b/Engine/Source/Developer/Datasmith/DatasmithExporterUI/Private/DirectLinkUI.h
--- a/Engine/Source/Developer/Datasmith/DatasmithExporterUI/Private/DirectLinkUI.h
+++ b/Engine/Source/Developer/Datasmith/DatasmithExporterUI/Private/DirectLinkUI.h
@@ -21,12 +21,16 @@ public:
 
 private: 
 	void OnCacheDirectoryChanged(const FString& InNewCacheDirectory);
+	FString OnCacheDirectoryReset();
+
+	void SaveCacheDirectory( const FString& InCacheDir, bool bInDefaultCacheDir );
 
 	// To accessed from the game thread only
 	TWeakPtr<SWindow> DirectLinkWindow;
 	
 	FCriticalSection CriticalSectionCacheDirectory;
 	FString DirectLinkCacheDirectory;
+	FString DefaultDirectLinkCacheDirectory;
 
 	// Used to protect the caller to GetDirectLinkCacheDirectory from a pontential race condition
 	FString LastReturnedCacheDirectory;
diff --git a/Engine/Source/Developer/Datasmith/DatasmithExporterUI/Private/Widgets/SDirectLinkStreamManager.cpp b/Engine/Source/Developer/Datasmith/DatasmithExporterUI/Private/Widgets/SDirectLinkStreamManager.cpp
--- a/Engine/Source/Developer/Datasmith/DatasmithExporterUI/Private/Widgets/SDirectLinkStreamManager.cpp
+++ b/Engine/Source/Developer/Datasmith/DatasmithExporterUI/Private/Widgets/SDirectLinkStreamManager.cpp
@@ -202,6 +202,7 @@ void SDirectLinkStreamManager::Construct(const FArguments& InArgs, const TShared
 
 	DirectLinkCacheDirectory = InArgs._DefaultCacheDirectory;
 	OnCacheDirectoryChanged = InArgs._OnCacheDirectoryChanged;
+	OnCacheDirectoryReset = InArgs._OnCacheDirectoryReset;
 
 	TSharedRef<SHeaderRow> HeaderRow = SNew( SHeaderRow )
 		// Source
@@ -297,6 +298,24 @@ void SDirectLinkStreamManager::Construct(const FArguments& InArgs, const TShared
 			.Padding( 2.f )
 			.AutoWidth()
 			.HAlign( HAlign_Right )
+			[
+				SNew( SButton )
+				.OnClicked( this, &SDirectLinkStreamManager::OnResetCacheDirectoryClicked )
+				.ToolTipText( LOCTEXT("ResetCacheDirectory_Tooltip", "Reset cache directory") )
+				[
+					SNew( SVerticalBox )
+					+ SVerticalBox::Slot()
+					.VAlign( VAlign_Center )
+					[
+						SNew( STextBlock )
+						.Text( LOCTEXT("ResetCacheDirectory_Label", "Reset") )
+					]
+				]
+			]
+			+ SHorizontalBox::Slot()
+			.Padding( 2.f )
+			.AutoWidth()
+			.HAlign( HAlign_Right )
 			[
 				SNew( SButton )
 				.OnClicked( this, &SDirectLinkStreamManager::OnChangeCacheDirectoryClicked )
@@ -508,6 +527,16 @@ FText SDirectLinkStreamManager::GetCacheDirectory() const
 	return FText::FromString( DirectLinkCacheDirectory );
 }
 
+FReply SDirectLinkStreamManager::OnResetCacheDirectoryClicked()
+{
+	if( OnCacheDirectoryReset.IsBound() )
+	{
+		DirectLinkCacheDirectory = OnCacheDirectoryReset.Execute();
+	}
+
+	return FReply::Handled();
+}
+
 FReply SDirectLinkStreamManager::OnChangeCacheDirectoryClicked()
 {
 	IDesktopPlatform* DesktopPlatform = FDesktopPlatformModule::Get();
diff --git a/Engine/Source/Developer/Datasmith/DatasmithExporterUI/Private/Widgets/SDirectLinkStreamManager.h b/Engine/Source/Developer/Datasmith/DatasmithExporterUI/Private/Widgets/SDirectLinkStreamManager.h
--- a/Engine/Source/Developer/Datasmith/DatasmithExporterUI/Private/Widgets/SDirectLinkStreamManager.h
+++ b/Engine/Source/Developer/Datasmith/DatasmithExporterUI/Private/Widgets/SDirectLinkStreamManager.h
@@ -31,9 +31,11 @@ class SDirectLinkStreamManager : public SCompoundWidget
 {
 public:
 	DECLARE_DELEGATE_OneParam( FOnCacheDirectoryChanged, const FString& /** NewDirectory */ );
+	DECLARE_DELEGATE_RetVal( FString, FOnCacheDirectoryReset );
 
 	SLATE_BEGIN_ARGS(SDirectLinkStreamManager) {}
 		SLATE_EVENT( FOnCacheDirectoryChanged, OnCacheDirectoryChanged )
+		SLATE_EVENT( FOnCacheDirectoryReset, OnCacheDirectoryReset )
 		SLATE_ARGUMENT( FString, DefaultCacheDirectory )
 	SLATE_END_ARGS()
 
@@ -59,6 +61,7 @@ private:
 	FText GetCacheDirectory() const;
 
 	FReply OnChangeCacheDirectoryClicked();
+	FReply OnResetCacheDirectoryClicked();
 
 	EVisibility GetNoConnectionHintVisibility() const;
 
@@ -97,6 +100,7 @@ private:
 	FString DirectLinkCacheDirectory;
 
 	FOnCacheDirectoryChanged OnCacheDirectoryChanged;
+	FOnCacheDirectoryReset OnCacheDirectoryReset;
 
 	TSharedPtr<class SButton> ShowAdavancedSettingButton;
 	TSharedPtr<class SImage> ShowAdavancedSettingImage;
