Datasmith: fix double -> float conversion warning

diff --git a/Engine/Source/Developer/Datasmith/DatasmithFacade/Private/DatasmithFacadeAnimation.cpp b/Engine/Source/Developer/Datasmith/DatasmithFacade/Private/DatasmithFacadeAnimation.cpp
--- a/Engine/Source/Developer/Datasmith/DatasmithFacade/Private/DatasmithFacadeAnimation.cpp
+++ b/Engine/Source/Developer/Datasmith/DatasmithFacade/Private/DatasmithFacadeAnimation.cpp
@@ -75,7 +75,7 @@ EDatasmithFacadeCurveInterpMode FDatasmithFacadeTransformAnimation::GetCurveInte
 	return static_cast<EDatasmithFacadeCurveInterpMode>(GetDatasmithTransformAnimation()->GetCurveInterpMode(static_cast<EDatasmithTransformType>(TransformType)));
 }
 
-void FDatasmithFacadeTransformAnimation::GetFrame(EDatasmithFacadeAnimationTransformType TransformType, int32 Index, int32& OutFrameNumber, float& OutX, float& OutY, float& OutZ) const
+void FDatasmithFacadeTransformAnimation::GetFrame(EDatasmithFacadeAnimationTransformType TransformType, int32 Index, int32& OutFrameNumber, double& OutX, double& OutY, double& OutZ) const
 {
 	const FDatasmithTransformFrameInfo& FrameInfo = GetDatasmithTransformAnimation()->GetFrame(static_cast<EDatasmithTransformType>(TransformType), Index);
 	OutFrameNumber = FrameInfo.FrameNumber;
@@ -172,7 +172,7 @@ FDatasmithFacadeSubsequenceAnimation::FDatasmithFacadeSubsequenceAnimation(
 {
 }
 
-TSharedRef<IDatasmithSubsequenceAnimationElement> FDatasmithFacadeSubsequenceAnimation::GetDatasmithSubsequenceAnimation() const 
+TSharedRef<IDatasmithSubsequenceAnimationElement> FDatasmithFacadeSubsequenceAnimation::GetDatasmithSubsequenceAnimation() const
 {
 	return StaticCastSharedRef<IDatasmithSubsequenceAnimationElement>(InternalDatasmithElement);
 }
@@ -299,9 +299,9 @@ FDatasmithFacadeTransformAnimation* FDatasmithFacadeLevelSequence::GetNewTransfo
 	return nullptr;
 }
 
-FDatasmithFacadeVisibilityAnimation* FDatasmithFacadeLevelSequence::GetNewVisibilityAnimation(int32 InIndex) 
+FDatasmithFacadeVisibilityAnimation* FDatasmithFacadeLevelSequence::GetNewVisibilityAnimation(int32 InIndex)
 {
-	if (TSharedPtr<IDatasmithVisibilityAnimationElement> AnimElement = StaticCastSharedPtr<IDatasmithVisibilityAnimationElement>(GetDatasmithLevelSequence()->GetAnimation(InIndex))) 
+	if (TSharedPtr<IDatasmithVisibilityAnimationElement> AnimElement = StaticCastSharedPtr<IDatasmithVisibilityAnimationElement>(GetDatasmithLevelSequence()->GetAnimation(InIndex)))
 	{
 		return new FDatasmithFacadeVisibilityAnimation(AnimElement.ToSharedRef());
 	}
@@ -309,9 +309,9 @@ FDatasmithFacadeVisibilityAnimation* FDatasmithFacadeLevelSequence::GetNewVisibi
 	return nullptr;
 }
 
-FDatasmithFacadeSubsequenceAnimation* FDatasmithFacadeLevelSequence::GetNewSubsequenceAnimation(int32 InIndex) 
+FDatasmithFacadeSubsequenceAnimation* FDatasmithFacadeLevelSequence::GetNewSubsequenceAnimation(int32 InIndex)
 {
-	if (TSharedPtr<IDatasmithSubsequenceAnimationElement> AnimElement = StaticCastSharedPtr<IDatasmithSubsequenceAnimationElement>(GetDatasmithLevelSequence()->GetAnimation(InIndex))) 
+	if (TSharedPtr<IDatasmithSubsequenceAnimationElement> AnimElement = StaticCastSharedPtr<IDatasmithSubsequenceAnimationElement>(GetDatasmithLevelSequence()->GetAnimation(InIndex)))
 	{
 		return new FDatasmithFacadeSubsequenceAnimation(AnimElement.ToSharedRef());
 	}
diff --git a/Engine/Source/Developer/Datasmith/DatasmithFacade/Public/DatasmithFacadeAnimation.h b/Engine/Source/Developer/Datasmith/DatasmithFacade/Public/DatasmithFacadeAnimation.h
--- a/Engine/Source/Developer/Datasmith/DatasmithFacade/Public/DatasmithFacadeAnimation.h
+++ b/Engine/Source/Developer/Datasmith/DatasmithFacade/Public/DatasmithFacadeAnimation.h
@@ -95,7 +95,7 @@ public:
 	EDatasmithFacadeCurveInterpMode GetCurveInterpMode(EDatasmithFacadeAnimationTransformType TransformType) const;
 
 	/** Return the frame of the given transform type at the given index or an invalid frame if the index was out of bounds */
-	void GetFrame(EDatasmithFacadeAnimationTransformType TransformType, int32 Index, int32& OutFrameNumber, float& OutX, float& OutY, float& OutZ) const;
+	void GetFrame(EDatasmithFacadeAnimationTransformType TransformType, int32 Index, int32& OutFrameNumber, double& OutX, double& OutY, double& OutZ) const;
 
 	/** Remove the frame of the given transform type at the given index from the animation */
 	void RemoveFrame(EDatasmithFacadeAnimationTransformType TransformType, int32 Index);
diff --git a/Engine/Source/Programs/Enterprise/Datasmith/DatasmithFacadeCSharp/Private/DatasmithFacadeCSharp.cpp b/Engine/Source/Programs/Enterprise/Datasmith/DatasmithFacadeCSharp/Private/DatasmithFacadeCSharp.cpp
--- a/Engine/Source/Programs/Enterprise/Datasmith/DatasmithFacadeCSharp/Private/DatasmithFacadeCSharp.cpp
+++ b/Engine/Source/Programs/Enterprise/Datasmith/DatasmithFacadeCSharp/Private/DatasmithFacadeCSharp.cpp
@@ -4936,22 +4936,34 @@ SWIGEXPORT int SWIGSTDCALL CSharp_FDatasmithFacadeTransformAnimation_GetCurveInt
 }
 
 
-SWIGEXPORT void SWIGSTDCALL CSharp_FDatasmithFacadeTransformAnimation_GetFrame(void * jarg1, int jarg2, int jarg3, int * jarg4, float * jarg5, float * jarg6, float * jarg7) {
+SWIGEXPORT void SWIGSTDCALL CSharp_FDatasmithFacadeTransformAnimation_GetFrame(void * jarg1, int jarg2, int jarg3, int * jarg4, void * jarg5, void * jarg6, void * jarg7) {
   FDatasmithFacadeTransformAnimation *arg1 = (FDatasmithFacadeTransformAnimation *) 0 ;
   EDatasmithFacadeAnimationTransformType arg2 ;
   int32 arg3 ;
   int32 *arg4 = 0 ;
-  float *arg5 = 0 ;
-  float *arg6 = 0 ;
-  float *arg7 = 0 ;
+  double *arg5 = 0 ;
+  double *arg6 = 0 ;
+  double *arg7 = 0 ;
   
   arg1 = (FDatasmithFacadeTransformAnimation *)jarg1; 
   arg2 = (EDatasmithFacadeAnimationTransformType)jarg2; 
   arg3 = (int32)jarg3; 
   arg4 = (int32 *)jarg4; 
-  arg5 = (float *)jarg5; 
-  arg6 = (float *)jarg6; 
-  arg7 = (float *)jarg7; 
+  arg5 = (double *)jarg5;
+  if (!arg5) {
+    SWIG_CSharpSetPendingExceptionArgument(SWIG_CSharpArgumentNullException, "double & type is null", 0);
+    return ;
+  } 
+  arg6 = (double *)jarg6;
+  if (!arg6) {
+    SWIG_CSharpSetPendingExceptionArgument(SWIG_CSharpArgumentNullException, "double & type is null", 0);
+    return ;
+  } 
+  arg7 = (double *)jarg7;
+  if (!arg7) {
+    SWIG_CSharpSetPendingExceptionArgument(SWIG_CSharpArgumentNullException, "double & type is null", 0);
+    return ;
+  } 
   ((FDatasmithFacadeTransformAnimation const *)arg1)->GetFrame(arg2,arg3,*arg4,*arg5,*arg6,*arg7);
 }
 
diff --git a/Engine/Source/Programs/Enterprise/Datasmith/DatasmithFacadeCSharp/Private/DatasmithFacadeDLL.cpp b/Engine/Source/Programs/Enterprise/Datasmith/DatasmithFacadeCSharp/Private/DatasmithFacadeDLL.cpp
--- a/Engine/Source/Programs/Enterprise/Datasmith/DatasmithFacadeCSharp/Private/DatasmithFacadeDLL.cpp
+++ b/Engine/Source/Programs/Enterprise/Datasmith/DatasmithFacadeCSharp/Private/DatasmithFacadeDLL.cpp
@@ -11,7 +11,7 @@
 
 
 #if PLATFORM_WINDOWS
-// Begin Datasmith platform include gard.
+// Begin Datasmith platform include guard.
 #include "Windows/AllowWindowsPlatformTypes.h"
 
 BOOL WINAPI DllMain(
diff --git a/Engine/Source/Programs/Enterprise/Datasmith/DatasmithFacadeCSharp/Public/DatasmithFacadeCSharpPINVOKE.cs b/Engine/Source/Programs/Enterprise/Datasmith/DatasmithFacadeCSharp/Public/DatasmithFacadeCSharpPINVOKE.cs
--- a/Engine/Source/Programs/Enterprise/Datasmith/DatasmithFacadeCSharp/Public/DatasmithFacadeCSharpPINVOKE.cs
+++ b/Engine/Source/Programs/Enterprise/Datasmith/DatasmithFacadeCSharp/Public/DatasmithFacadeCSharpPINVOKE.cs
@@ -1386,7 +1386,7 @@ class DatasmithFacadeCSharpPINVOKE {
   public static extern int FDatasmithFacadeTransformAnimation_GetCurveInterpMode(global::System.Runtime.InteropServices.HandleRef jarg1, int jarg2);
 
   [global::System.Runtime.InteropServices.DllImport("DatasmithFacadeCSharp", EntryPoint="CSharp_FDatasmithFacadeTransformAnimation_GetFrame")]
-  public static extern void FDatasmithFacadeTransformAnimation_GetFrame(global::System.Runtime.InteropServices.HandleRef jarg1, int jarg2, int jarg3, out int jarg4, out float jarg5, out float jarg6, out float jarg7);
+  public static extern void FDatasmithFacadeTransformAnimation_GetFrame(global::System.Runtime.InteropServices.HandleRef jarg1, int jarg2, int jarg3, out int jarg4, global::System.Runtime.InteropServices.HandleRef jarg5, global::System.Runtime.InteropServices.HandleRef jarg6, global::System.Runtime.InteropServices.HandleRef jarg7);
 
   [global::System.Runtime.InteropServices.DllImport("DatasmithFacadeCSharp", EntryPoint="CSharp_FDatasmithFacadeTransformAnimation_RemoveFrame")]
   public static extern void FDatasmithFacadeTransformAnimation_RemoveFrame(global::System.Runtime.InteropServices.HandleRef jarg1, int jarg2, int jarg3);
diff --git a/Engine/Source/Programs/Enterprise/Datasmith/DatasmithFacadeCSharp/Public/FDatasmithFacadeTransformAnimation.cs b/Engine/Source/Programs/Enterprise/Datasmith/DatasmithFacadeCSharp/Public/FDatasmithFacadeTransformAnimation.cs
--- a/Engine/Source/Programs/Enterprise/Datasmith/DatasmithFacadeCSharp/Public/FDatasmithFacadeTransformAnimation.cs
+++ b/Engine/Source/Programs/Enterprise/Datasmith/DatasmithFacadeCSharp/Public/FDatasmithFacadeTransformAnimation.cs
@@ -56,8 +56,9 @@ public class FDatasmithFacadeTransformAnimation : FDatasmithFacadeBaseAnimation
     return ret;
   }
 
-  public void GetFrame(EDatasmithFacadeAnimationTransformType TransformType, int Index, out int OutFrameNumber, out float OutX, out float OutY, out float OutZ) {
-    DatasmithFacadeCSharpPINVOKE.FDatasmithFacadeTransformAnimation_GetFrame(swigCPtr, (int)TransformType, Index, out OutFrameNumber, out OutX, out OutY, out OutZ);
+  public void GetFrame(EDatasmithFacadeAnimationTransformType TransformType, int Index, out int OutFrameNumber, SWIGTYPE_p_double OutX, SWIGTYPE_p_double OutY, SWIGTYPE_p_double OutZ) {
+    DatasmithFacadeCSharpPINVOKE.FDatasmithFacadeTransformAnimation_GetFrame(swigCPtr, (int)TransformType, Index, out OutFrameNumber, SWIGTYPE_p_double.getCPtr(OutX), SWIGTYPE_p_double.getCPtr(OutY), SWIGTYPE_p_double.getCPtr(OutZ));
+    if (DatasmithFacadeCSharpPINVOKE.SWIGPendingException.Pending) throw DatasmithFacadeCSharpPINVOKE.SWIGPendingException.Retrieve();
   }
 
   public void RemoveFrame(EDatasmithFacadeAnimationTransformType TransformType, int Index) {
