Fixing DatasmithCSharpFacade compilation on Mac

JeanMichel.Dignard

diff --git a/Engine/Source/Developer/Datasmith/DatasmithFacade/Private/DatasmithFacadeUtils.cpp b/Engine/Source/Developer/Datasmith/DatasmithFacade/Private/DatasmithFacadeUtils.cpp
new file mode 100644
--- /dev/null
+++ b/Engine/Source/Developer/Datasmith/DatasmithFacade/Private/DatasmithFacadeUtils.cpp
@@ -0,0 +1,30 @@
+// Copyright Epic Games, Inc. All Rights Reserved.
+
+#include "DatasmithFacadeUtils.h"
+#include "DatasmithUtils.h"
+
+FDatasmithFacadeUniqueNameProvider::FDatasmithFacadeUniqueNameProvider()
+	:InternalNameProvider(MakeUnique<FDatasmithUniqueNameProvider>())
+{}
+
+
+const TCHAR* FDatasmithFacadeUniqueNameProvider::GenerateUniqueName(const TCHAR* BaseName)
+{
+	CachedGeneratedName = InternalNameProvider->GenerateUniqueName(BaseName);
+	return *CachedGeneratedName;
+}
+
+void FDatasmithFacadeUniqueNameProvider::Reserve( int32 NumberOfName )
+{
+	InternalNameProvider->Reserve(NumberOfName);
+}
+
+void FDatasmithFacadeUniqueNameProvider::AddExistingName(const TCHAR* Name)
+{
+	InternalNameProvider->AddExistingName(Name);
+}
+
+void FDatasmithFacadeUniqueNameProvider::RemoveExistingName(const TCHAR* Name)
+{
+	InternalNameProvider->RemoveExistingName(Name);
+}
diff --git a/Engine/Source/Developer/Datasmith/DatasmithFacade/Public/DatasmithFacadeUtils.h b/Engine/Source/Developer/Datasmith/DatasmithFacade/Public/DatasmithFacadeUtils.h
--- a/Engine/Source/Developer/Datasmith/DatasmithFacade/Public/DatasmithFacadeUtils.h
+++ b/Engine/Source/Developer/Datasmith/DatasmithFacade/Public/DatasmithFacadeUtils.h
@@ -3,43 +3,41 @@
 #pragma once
 
 #include "CoreMinimal.h"
-#include "DatasmithUtils.h"
 
+class FDatasmithUniqueNameProvider;
 
 class DATASMITHFACADE_API FDatasmithFacadeUniqueNameProvider
 {
 public:
+	FDatasmithFacadeUniqueNameProvider();
+	
 	/**
 	 * Generates a unique name
 	 * @param BaseName Name that will be suffixed with an index to be unique
 	 * @return TCHAR* Unique name. Calling "Contains()" with this name will be false. 
 	 *                Pointer is only valid until the next GenerateUniqueName() call.
 	 */
-	const TCHAR* GenerateUniqueName(const TCHAR* BaseName)
-	{
-		CachedGeneratedName = InternalNameProvider.GenerateUniqueName(BaseName);
-		return *CachedGeneratedName;
-	}
+	const TCHAR* GenerateUniqueName(const TCHAR* BaseName);
 
 	/**
 	 * Reserves space in the internal data structures to contain at least the number of name specified.
 	 * @param NumberOfName The number of name to reserve for.
 	 */
-	void Reserve( int32 NumberOfName ) { InternalNameProvider.Reserve(NumberOfName); }
+	void Reserve( int32 NumberOfName );
 
 	/**
 	 * Register a name as known
 	 * @param Name name to register
 	 */
-	void AddExistingName(const TCHAR* Name) { InternalNameProvider.AddExistingName(Name); }
+	void AddExistingName(const TCHAR* Name);
 	
 	/**
 	 * Remove a name from the list of existing name
 	 * @param Name name to unregister
 	 */
-	void RemoveExistingName(const TCHAR* Name) { InternalNameProvider.RemoveExistingName(Name); }
+	void RemoveExistingName(const TCHAR* Name);
 
 private:
-	FDatasmithUniqueNameProvider InternalNameProvider;
+	TUniquePtr<FDatasmithUniqueNameProvider> InternalNameProvider;
 	FString CachedGeneratedName;
-};
\ No newline at end of file
+};
