Datasmith: New mesh serialization, MeshDescription based

diff --git a/Engine/Plugins/Enterprise/DatasmithImporter/Source/DatasmithImporter/Private/DatasmithImporter.cpp b/Engine/Plugins/Enterprise/DatasmithImporter/Source/DatasmithImporter/Private/DatasmithImporter.cpp
--- a/Engine/Plugins/Enterprise/DatasmithImporter/Source/DatasmithImporter/Private/DatasmithImporter.cpp
+++ b/Engine/Plugins/Enterprise/DatasmithImporter/Source/DatasmithImporter/Private/DatasmithImporter.cpp
@@ -253,6 +253,11 @@ void FDatasmithImporter::ImportStaticMeshes( FDatasmithImportContext& ImportCont
 		return;
 	}
 
+	if (!ImportContext.AssetsContext.StaticMeshesFinalPackage || ImportContext.AssetsContext.StaticMeshesFinalPackage->GetFName() == NAME_None || ImportContext.SceneTranslator == nullptr)
+	{
+		return;
+	}
+
 	TRACE_CPUPROFILER_EVENT_SCOPE(FDatasmithImporter::ImportStaticMeshes);
 
 	TUniquePtr<FScopedSlowTask> ProgressPtr;
@@ -266,10 +271,7 @@ void FDatasmithImporter::ImportStaticMeshes( FDatasmithImportContext& ImportCont
 	TMap<TSharedRef<IDatasmithMeshElement>, TFuture<FDatasmithMeshElementPayload*>> MeshElementPayloads;
 
 	FDatasmithTranslatorCapabilities TranslatorCapabilities;
-	if (ImportContext.SceneTranslator)
-	{
-		ImportContext.SceneTranslator->Initialize(TranslatorCapabilities);
-	}
+	ImportContext.SceneTranslator->Initialize(TranslatorCapabilities);
 
 	// Parallelize loading by doing a first pass to send translator loading into async task
 	if (TranslatorCapabilities.bParallelLoadStaticMeshSupported)
@@ -281,11 +283,6 @@ void FDatasmithImporter::ImportStaticMeshes( FDatasmithImportContext& ImportCont
 				ImportContext.bUserCancelled = true;
 			}
 
-			if (!ImportContext.AssetsContext.StaticMeshesFinalPackage || ImportContext.AssetsContext.StaticMeshesFinalPackage->GetFName() == NAME_None || ImportContext.SceneTranslator == nullptr)
-			{
-				continue;
-			}
-
 			TSharedRef<IDatasmithMeshElement> MeshElement = ImportContext.FilteredScene->GetMesh( MeshIndex ).ToSharedRef();
 
 			UStaticMesh*& ImportedStaticMesh = ImportContext.ImportedStaticMeshes.FindOrAdd( MeshElement );
diff --git a/Engine/Plugins/Enterprise/DatasmithImporter/Source/DatasmithImporter/Public/DatasmithStaticMeshImporter.h b/Engine/Plugins/Enterprise/DatasmithImporter/Source/DatasmithImporter/Public/DatasmithStaticMeshImporter.h
--- a/Engine/Plugins/Enterprise/DatasmithImporter/Source/DatasmithImporter/Public/DatasmithStaticMeshImporter.h
+++ b/Engine/Plugins/Enterprise/DatasmithImporter/Source/DatasmithImporter/Public/DatasmithStaticMeshImporter.h
@@ -11,7 +11,6 @@
 
 struct FDatasmithAssetsImportContext;
 struct FDatasmithImportContext;
-class FDatasmithLegacyImporter;
 class FDatasmithMesh;
 struct FDatasmithStaticMeshImportOptions;
 struct FMeshDescription;
@@ -20,7 +19,6 @@ struct FOverlappingCorners;
 class IDatasmithMeshElement;
 class IDatasmithScene;
 class IMeshUtilities;
-class UDatasmithMesh;
 class UObject;
 class UPackage;
 class UStaticMesh;
@@ -29,9 +27,6 @@ struct FDatasmithMeshElementPayload;
 class DATASMITHIMPORTER_API FDatasmithStaticMeshImporter
 {
 public:
-	/** DEPRECATED. Imports a static mesh from a Mesh Element */
-	static UStaticMesh* ImportStaticMesh( const TSharedRef< IDatasmithMeshElement > MeshElement, FDatasmithLegacyImporter& DatasmithImporter, EObjectFlags ObjectFlags, const FDatasmithStaticMeshImportOptions& ImportOptions, FDatasmithAssetsImportContext& AssetsContext, UStaticMesh* ExistingMesh );
-
 	/** Imports a static mesh from a Mesh Element */
 	static UStaticMesh* ImportStaticMesh( const TSharedRef< IDatasmithMeshElement > MeshElement, FDatasmithMeshElementPayload& Payload, EObjectFlags ObjectFlags, const FDatasmithStaticMeshImportOptions& ImportOptions, FDatasmithAssetsImportContext& AssetsContext, UStaticMesh* ExistingMesh );
 
@@ -43,7 +38,7 @@ public:
 
 	/** Setup a UStaticMesh from an IDatasmithMeshElement. */
 	static void SetupStaticMesh( FDatasmithAssetsImportContext& AssetsContext, TSharedRef< IDatasmithMeshElement > MeshElement, UStaticMesh* StaticMesh, const FDatasmithStaticMeshImportOptions& StaticMeshImportOptions, float LightmapWeight );
-	
+
 	/** Cleanup any invalid data in mesh descriptions that might cause the editor to crash or behave erratically (i.e. Having vertex position with NaN values). */
 	static void CleanupMeshDescriptions(TArray< FMeshDescription >& MeshDescriptions);
 
diff --git a/Engine/Plugins/Enterprise/DatasmithImporter/Source/DatasmithNativeTranslator/Private/DatasmithNativeTranslator.cpp b/Engine/Plugins/Enterprise/DatasmithImporter/Source/DatasmithNativeTranslator/Private/DatasmithNativeTranslator.cpp
--- a/Engine/Plugins/Enterprise/DatasmithImporter/Source/DatasmithNativeTranslator/Private/DatasmithNativeTranslator.cpp
+++ b/Engine/Plugins/Enterprise/DatasmithImporter/Source/DatasmithNativeTranslator/Private/DatasmithNativeTranslator.cpp
@@ -19,6 +19,8 @@
 #include "UObject/StrongObjectPtr.h"
 
 
+DEFINE_LOG_CATEGORY(LogDatasmithNativeTranslator);
+
 FString FDatasmithNativeTranslator::ResolveFilePath(const FString& FilePath, const TArray<FString>& ResourcePaths)
 {
 	if (FPaths::IsRelative(FilePath) && !FPaths::FileExists(FilePath))
@@ -119,63 +121,11 @@ bool FDatasmithNativeTranslator::LoadScene(TSharedRef<IDatasmithScene> OutScene)
 
 namespace DatasmithNativeTranslatorImpl
 {
-	struct FDatasmithMeshInternal
-	{
-		bool bIsCollisionMesh = false;
-		TArray< FDatasmithMeshSourceModel > SourceModels;
-	};
-
-	TArray< FDatasmithMeshInternal > GetDatasmithMeshFromMeshPath( const TCHAR* MeshPath )
-	{
-		TArray< FDatasmithMeshInternal > Result;
-
-		TUniquePtr<FArchive> Archive( IFileManager::Get().CreateFileReader(MeshPath) );
-		if ( !Archive.IsValid() )
-		{
-			return Result;
-		}
-
-		int32 NumMeshes = 0;
-		*Archive << NumMeshes;
-
-		UDatasmithMesh* DatasmithMesh = nullptr;
-		{
-			// Make sure the new UDatasmithMesh object is not created while a garbage collection is performed
-			FGCScopeGuard GCGuard;
-			// Setting the RF_Standalone bitmask on the new UDatasmithMesh object, to make sure it is not garbage collected
-			// while loading and processing the udsmesh file. This can happen on very big meshes (5M+ triangles)
-			DatasmithMesh = NewObject< UDatasmithMesh >(GetTransientPackage(), NAME_None, RF_Standalone);
-		}
-
-		// Currently we only have 1 mesh per file. If there's a second mesh, it will be a CollisionMesh
-		while ( NumMeshes-- > 0)
-		{
-			TArray< uint8 > Bytes;
-			*Archive << Bytes;
-
-			FMemoryReader MemoryReader( Bytes, true );
-			MemoryReader.ArIgnoreClassRef = false;
-			MemoryReader.ArIgnoreArchetypeRef = false;
-			MemoryReader.SetWantBinaryPropertySerialization(true);
-			DatasmithMesh->Serialize( MemoryReader );
-
-			FDatasmithMeshInternal& SourceModels = Result.AddDefaulted_GetRef();
-
-			SourceModels.bIsCollisionMesh = DatasmithMesh->bIsCollisionMesh;
-			SourceModels.SourceModels = MoveTemp(DatasmithMesh->SourceModels);
-		}
-
-		// Tell the garbage collector DatasmithMesh can now be deleted.
-		DatasmithMesh->ClearInternalFlags(EInternalObjectFlags::Async);
-		DatasmithMesh->ClearFlags(RF_Standalone);
-
-		return Result;
-	}
-
-	TOptional<FMeshDescription> ExtractMeshDescription(FDatasmithMeshSourceModel& DSSourceModel)
+	TOptional<FMeshDescription> ExtractToMeshDescription(FDatasmithMeshSourceModel& SourceModel)
 	{
 		FRawMesh RawMesh;
-		DSSourceModel.RawMeshBulkData.LoadRawMesh( RawMesh );
+		SourceModel.RawMeshBulkData.LoadRawMesh( RawMesh );
+
 		if ( !RawMesh.IsValid() )
 		{
 			return {};
@@ -218,8 +168,84 @@ namespace DatasmithNativeTranslatorImpl
 		return MeshDescription;
 	}
 
+
+	TArray< FDatasmithMeshModels > GetDatasmithMeshFromMeshPath( const FString& MeshPath )
+	{
+		TArray< FDatasmithMeshModels > Result;
+
+		TUniquePtr<FArchive> Archive( IFileManager::Get().CreateFileReader(*MeshPath) );
+		if ( !Archive.IsValid() )
+		{
+			UE_LOG(LogDatasmithNativeTranslator, Warning, TEXT("Cannot read file %s"), *MeshPath);
+			return Result;
+		}
+
+		int32 LeagacyNumMeshesCount = 0;
+		*Archive << LeagacyNumMeshesCount;
+
+		if (LeagacyNumMeshesCount == 0)
+		{
+			FDatasmithPackedMeshes Pack;
+			*Archive << Pack;
+
+			if (!Archive->IsError())
+			{
+				for (FDatasmithMeshModels& Mesh : Pack.MeshesToExport)
+				{
+					FDatasmithMeshModels& MeshInternal = Result.AddDefaulted_GetRef();
+					MeshInternal.bIsCollisionMesh = Mesh.bIsCollisionMesh;
+					MeshInternal.SourceModels = MoveTemp(Mesh.SourceModels);
+				}
+			}
+			else
+			{
+				UE_LOG(LogDatasmithNativeTranslator, Warning, TEXT("Failed to read meshes from %s"), *MeshPath);
+			}
+		}
+		else
+		{
+			UDatasmithMesh* DatasmithMesh = nullptr;
+			{
+				// Make sure the new UDatasmithMesh object is not created while a garbage collection is performed
+				FGCScopeGuard GCGuard;
+				// Setting the RF_Standalone bitmask on the new UDatasmithMesh object, to make sure it is not garbage collected
+				// while loading and processing the udsmesh file. This can happen on very big meshes (5M+ triangles)
+				DatasmithMesh = NewObject< UDatasmithMesh >(GetTransientPackage(), NAME_None, RF_Standalone);
+			}
+
+			// Currently we only have 1 mesh per file. If there's a second mesh, it will be a CollisionMesh
+			while ( LeagacyNumMeshesCount-- > 0)
+			{
+				TArray< uint8 > Bytes;
+				*Archive << Bytes;
+
+				FMemoryReader MemoryReader( Bytes, true );
+				MemoryReader.ArIgnoreClassRef = false;
+				MemoryReader.ArIgnoreArchetypeRef = false;
+				MemoryReader.SetWantBinaryPropertySerialization(true);
+				DatasmithMesh->Serialize( MemoryReader );
+
+				FDatasmithMeshModels& MeshInternal = Result.AddDefaulted_GetRef();
+				MeshInternal.bIsCollisionMesh = DatasmithMesh->bIsCollisionMesh;
+				for (FDatasmithMeshSourceModel& SourceModel : DatasmithMesh->SourceModels)
+				{
+					if (TOptional<FMeshDescription> OptionalMesh = ExtractToMeshDescription(SourceModel))
+					{
+						MeshInternal.SourceModels.Add(MoveTemp(*OptionalMesh));
+					}
+				}
+			}
+
+			// Tell the garbage collector DatasmithMesh can now be deleted.
+			DatasmithMesh->ClearInternalFlags(EInternalObjectFlags::Async);
+			DatasmithMesh->ClearFlags(RF_Standalone);
+		}
+
+		return Result;
+	}
 } // ns DatasmithNativeTranslatorImpl
 
+
 bool FDatasmithNativeTranslator::LoadStaticMesh(const TSharedRef<IDatasmithMeshElement> MeshElement, FDatasmithMeshElementPayload& OutMeshPayload)
 {
 	TRACE_CPUPROFILER_EVENT_SCOPE(FDatasmithNativeTranslator::LoadStaticMesh);
@@ -227,38 +253,22 @@ bool FDatasmithNativeTranslator::LoadStaticMesh(const TSharedRef<IDatasmithMeshE
 	using namespace DatasmithNativeTranslatorImpl;
 
 	FString FilePath = MeshElement->GetFile();
-	if (!FPaths::FileExists(FilePath))
-	{
-		return false;
-	}
 
-	int32 ExtractionFailure = 0;
-	for (FDatasmithMeshInternal& DatasmithMesh : GetDatasmithMeshFromMeshPath( *FilePath ))
+	for (FDatasmithMeshModels& DatasmithMesh : GetDatasmithMeshFromMeshPath( FilePath ))
 	{
 		if (DatasmithMesh.bIsCollisionMesh)
 		{
-			for (FDatasmithMeshSourceModel& SourceModel : DatasmithMesh.SourceModels)
+			for (FMeshDescription& MeshDescription : DatasmithMesh.SourceModels)
 			{
-				FRawMesh RawMesh;
-				SourceModel.RawMeshBulkData.LoadRawMesh( RawMesh );
-				if ( RawMesh.VertexPositions.Num() > 0 )
-				{
-					OutMeshPayload.CollisionPointCloud = LWC::ConvertArrayType<FVector>(RawMesh.VertexPositions);	// LWC_TODO: Perf pessimization. Was MoveTemp.
-					break;
-				}
-				ExtractionFailure++;
+				OutMeshPayload.CollisionMesh = MoveTemp(MeshDescription);
+				break;
 			}
 		}
 		else
 		{
-			for (FDatasmithMeshSourceModel& SourceModel : DatasmithMesh.SourceModels)
+			for (FMeshDescription& SourceModel : DatasmithMesh.SourceModels)
 			{
-				if (TOptional<FMeshDescription> Mesh = ExtractMeshDescription(SourceModel))
-				{
-					OutMeshPayload.LodMeshes.Add(MoveTemp(Mesh.GetValue()));
-					continue;
-				}
-				ExtractionFailure++;
+				OutMeshPayload.LodMeshes.Add(MoveTemp(SourceModel));
 			}
 		}
 
@@ -266,7 +276,7 @@ bool FDatasmithNativeTranslator::LoadStaticMesh(const TSharedRef<IDatasmithMeshE
 		DatasmithMesh.SourceModels.Reset();
 	}
 
-	return ExtractionFailure == 0;
+	return OutMeshPayload.LodMeshes.Num() != 0;
 }
 
 bool FDatasmithNativeTranslator::LoadLevelSequence(const TSharedRef<IDatasmithLevelSequenceElement> LevelSequenceElement, FDatasmithLevelSequencePayload& OutLevelSequencePayload)
diff --git a/Engine/Plugins/Enterprise/DatasmithImporter/Source/DatasmithNativeTranslator/Public/DatasmithNativeTranslator.h b/Engine/Plugins/Enterprise/DatasmithImporter/Source/DatasmithNativeTranslator/Public/DatasmithNativeTranslator.h
--- a/Engine/Plugins/Enterprise/DatasmithImporter/Source/DatasmithNativeTranslator/Public/DatasmithNativeTranslator.h
+++ b/Engine/Plugins/Enterprise/DatasmithImporter/Source/DatasmithNativeTranslator/Public/DatasmithNativeTranslator.h
@@ -5,6 +5,9 @@
 #include "CoreMinimal.h"
 #include "DatasmithTranslator.h"
 
+#include "Logging/LogMacros.h"
+DECLARE_LOG_CATEGORY_EXTERN(LogDatasmithNativeTranslator, Log, All);
+
 class DATASMITHNATIVETRANSLATOR_API FDatasmithNativeTranslator : public IDatasmithTranslator
 {
 public:
diff --git a/Engine/Source/Developer/Datasmith/DatasmithExporter/Private/DatasmithMeshExporter.cpp b/Engine/Source/Developer/Datasmith/DatasmithExporter/Private/DatasmithMeshExporter.cpp
--- a/Engine/Source/Developer/Datasmith/DatasmithExporter/Private/DatasmithMeshExporter.cpp
+++ b/Engine/Source/Developer/Datasmith/DatasmithExporter/Private/DatasmithMeshExporter.cpp
@@ -2,6 +2,7 @@
 
 #include "DatasmithMeshExporter.h"
 
+#include "DatasmithCore.h"
 #include "DatasmithExporterManager.h"
 #include "DatasmithMesh.h"
 #include "DatasmithMeshUObject.h"
@@ -23,6 +24,22 @@
  */
 constexpr int32 DSMeshGarbageCollectionThreshold{ 2000 };
 
+struct FDatasmithMeshExporterOptions
+{
+	FDatasmithMeshExporterOptions( const FString& InFullPath, FDatasmithMesh& InMesh, EDSExportLightmapUV InLightmapUV, FDatasmithMesh* InCollisionMesh = nullptr )
+		: MeshFullPath( InFullPath )
+		, Mesh( InMesh )
+		, LightmapUV( InLightmapUV )
+		, CollisionMesh( InCollisionMesh )
+	{}
+
+	FString MeshFullPath;
+	FDatasmithMesh& Mesh;
+	EDSExportLightmapUV LightmapUV;
+	FDatasmithMesh* CollisionMesh;
+};
+void CreateDefaultUVs(FDatasmithMesh& DatasmithMesh);
+
 /**
  *	Implementation class of the DatasmithMeshExporter
  *	We use a lockfree UDatasmithMesh pool to avoid creating new UObject when exporting and reduces our memory footprint.
@@ -30,20 +47,6 @@ constexpr int32 DSMeshGarbageCollectionThreshold{ 2000 };
 class FDatasmithMeshExporterImpl
 {
 public:
-	struct FDatasmithMeshExporterOptions
-	{
-		FDatasmithMeshExporterOptions( const FString& InFullPath, FDatasmithMesh& InMesh, EDSExportLightmapUV InLightmapUV, FDatasmithMesh* InCollisionMesh = nullptr )
-			: MeshFullPath( InFullPath )
-			, Mesh( InMesh )
-			, LightmapUV( InLightmapUV )
-			, CollisionMesh( InCollisionMesh )
-		{}
-
-		FString MeshFullPath;
-		FDatasmithMesh& Mesh;
-		EDSExportLightmapUV LightmapUV;
-		FDatasmithMesh* CollisionMesh;
-	};
 
 	FDatasmithMeshExporterImpl() :
 		UniqueID(FGuid::NewGuid())
@@ -65,9 +68,6 @@ public:
 	bool ExportMeshes( const FDatasmithMeshExporterOptions& ExporterOptions, FMD5Hash& OutHash );
 	void PostExport(const FDatasmithMesh& DatasmithMesh, TSharedPtr< IDatasmithMeshElement >& MeshElement);
 
-	void CreateDefaultUVs(FDatasmithMesh& DatasmithMesh);
-	void RegisterStaticMeshAttributes(FMeshDescription& MeshDescription);
-
 	FString LastError;
 
 private:
@@ -114,28 +114,67 @@ private:
 
 TAtomic<int32> FDatasmithMeshExporterImpl::NumberOfUMeshPendingGC(0);
 
-FDatasmithMeshExporter::FDatasmithMeshExporter()
-{
-	Impl = new FDatasmithMeshExporterImpl();
-}
 
-FDatasmithMeshExporter::~FDatasmithMeshExporter()
+class FDatasmithMeshExporterImpl2
 {
-	delete Impl;
-	Impl = nullptr;
-}
+public:
+	bool DoExport( TSharedPtr< IDatasmithMeshElement >& MeshElement, const FDatasmithMeshExporterOptions& ExportOptions );
+	FString LastError;
 
-FString GetMeshFilePath( const TCHAR* Filepath, const TCHAR* Filename )
+private:
+	bool WriteMeshFile(const FDatasmithMeshExporterOptions& ExporterOptions, FMD5Hash& OutHash);
+};
+
+bool FDatasmithMeshExporterImpl2::DoExport(TSharedPtr< IDatasmithMeshElement >& MeshElement, const FDatasmithMeshExporterOptions& ExportOptions)
 {
-	FString NormalizedFilepath = Filepath;
-	FPaths::NormalizeDirectoryName( NormalizedFilepath );
+	FDatasmithMesh& Mesh = ExportOptions.Mesh;
 
-	FString NormalizedFilename = Filename;
-	FPaths::NormalizeFilename( NormalizedFilename );
+	// If the mesh doesn't have a name, use the filename as its name
+	if ( FCString::Strlen( Mesh.GetName() ) == 0 )
+	{
+		Mesh.SetName( *FPaths::GetBaseFilename( ExportOptions.MeshFullPath ) );
+	}
 
-	return FPaths::Combine( *NormalizedFilepath, FPaths::SetExtension( NormalizedFilename, UDatasmithMesh::GetFileExtension() ) );
+	// make sure we have at least one UV channels on each source models
+	CreateDefaultUVs( Mesh );
+	for ( int32 LODIndex = 0; LODIndex < Mesh.GetLODsCount(); ++LODIndex )
+	{
+		if ( FDatasmithMesh* LODMesh = Mesh.GetLOD( LODIndex ) )
+		{
+			CreateDefaultUVs( *LODMesh );
+		}
+	}
+
+	FMD5Hash Hash;
+	if ( WriteMeshFile( ExportOptions, Hash ) )
+	{
+		// If no existing MeshElement provided, create one.
+		if ( !MeshElement )
+		{
+			MeshElement = FDatasmithSceneFactory::CreateMesh( Mesh.GetName() );
+		}
+		MeshElement->SetFile( *ExportOptions.MeshFullPath );
+		MeshElement->SetFileHash( Hash );
+
+		FBox Extents = Mesh.GetExtents();
+		float Width = Extents.Max[0] - Extents.Min[0];
+		float Height = Extents.Max[2] - Extents.Min[2];
+		float Depth = Extents.Max[1] - Extents.Min[1];
+
+		MeshElement->SetDimensions( Mesh.ComputeArea(), Width, Height, Depth );
+		MeshElement->SetLightmapSourceUV( Mesh.GetLightmapSourceUVChannel() );
+
+		return true;
+	}
+	else
+	{
+		UE_LOG(LogDatasmith, Warning, TEXT("Cannot export mesh %s"), Mesh.GetName());
+	}
+
+	return false;
 }
 
+
 bool FDatasmithMeshExporterImpl::ExportMeshes( const FDatasmithMeshExporterOptions& ExporterOptions, FMD5Hash& OutHash )
 {
 	TArray<TSharedPtr<UDatasmithMesh>, TInlineAllocator<2>> MeshesToExport;
@@ -188,19 +227,93 @@ bool FDatasmithMeshExporterImpl::ExportMeshes( const FDatasmithMeshExporterOptio
 		*Archive << Bytes;
 	}
 	OutHash.Set( MD5 );
+	return true;
+}
 
-	Archive->Close();
 
-	return true;
+
+FDatasmithMeshExporter::FDatasmithMeshExporter()
+	: Impl(MakeUnique<FDatasmithMeshExporterImpl>())
+	, Impl2(MakeUnique<FDatasmithMeshExporterImpl2>())
+{}
+
+FDatasmithMeshExporter::~FDatasmithMeshExporter() = default;
+
+FString GetMeshFilePath( const TCHAR* Filepath, const TCHAR* Filename )
+{
+	FString NormalizedFilepath = Filepath;
+	FPaths::NormalizeDirectoryName( NormalizedFilepath );
+
+	FString NormalizedFilename = Filename;
+	FPaths::NormalizeFilename( NormalizedFilename );
+
+	return FPaths::Combine( *NormalizedFilepath, FPaths::SetExtension( NormalizedFilename, UDatasmithMesh::GetFileExtension() ) );
+}
+
+
+bool FDatasmithMeshExporterImpl2::WriteMeshFile( const FDatasmithMeshExporterOptions& ExporterOptions, FMD5Hash& OutHash )
+{
+	FDatasmithPackedMeshes Pack;
+
+	auto PackMeshModels = [&](FDatasmithMesh& Mesh, bool bIsCollisionMesh) -> FDatasmithMeshModels
+	{
+		FDatasmithMeshModels Models;
+		Models.bIsCollisionMesh = bIsCollisionMesh;
+		Models.MeshName = Mesh.GetName();
+
+		Models.SourceModels.Reserve(Mesh.GetLODsCount() + 1);
+		FMeshDescription& BaseMeshDescription = Models.SourceModels.AddDefaulted_GetRef();
+		FDatasmithMeshUtils::ToMeshDescription(Mesh, BaseMeshDescription);
+
+		for (int32 LodIndex = 0; LodIndex < Mesh.GetLODsCount(); ++LodIndex)
+		{
+			if (const FDatasmithMesh* LodMesh = Mesh.GetLOD(LodIndex))
+			{
+				FMeshDescription& LodMeshDescription = Models.SourceModels.AddDefaulted_GetRef();
+				FDatasmithMeshUtils::ToMeshDescription(*LodMesh, LodMeshDescription);
+			}
+		}
+
+		return Models;
+	};
+
+	Pack.MeshesToExport.Add(PackMeshModels(ExporterOptions.Mesh, false));
+	if (ExporterOptions.CollisionMesh)
+	{
+		Pack.MeshesToExport.Add(PackMeshModels(*ExporterOptions.CollisionMesh, true));
+	}
+
+	TUniquePtr<FArchive> Archive( IFileManager::Get().CreateFileWriter(*ExporterOptions.MeshFullPath) );
+	if (!Archive.IsValid())
+	{
+		LastError = FString::Printf( TEXT( "Failed writing to file %s" ), *ExporterOptions.MeshFullPath );
+		return false;
+	}
+
+	int32 LegacyMeshCount = 0;
+	*Archive << LegacyMeshCount; // the legacy importer expect a mesh count on the first bytes. Just in case a new file would end up parsed by the legacy code...
+	*Archive << Pack;
+
+	OutHash = Pack.OutHash;
+	return !Archive->IsError();
 }
 
-TSharedPtr< IDatasmithMeshElement > FDatasmithMeshExporter::ExportToUObject( const TCHAR* Filepath, const TCHAR* Filename, FDatasmithMesh& Mesh, FDatasmithMesh* CollisionMesh, EDSExportLightmapUV LightmapUV )
+
+
+TSharedPtr< IDatasmithMeshElement > FDatasmithMeshExporter::ExportToUObject(const TCHAR* Filepath, const TCHAR* Filename, FDatasmithMesh& Mesh, FDatasmithMesh* CollisionMesh, EDSExportLightmapUV LightmapUV, bool NewFormat/*=false */)
 {
 	FString FullPath( GetMeshFilePath( Filepath, Filename ) );
-	FDatasmithMeshExporterImpl::FDatasmithMeshExporterOptions ExportOptions( MoveTemp( FullPath ), Mesh, LightmapUV, CollisionMesh );
 
 	TSharedPtr<IDatasmithMeshElement> ExportedMeshElement;
-	Impl->DoExport( ExportedMeshElement, ExportOptions );
+	FDatasmithMeshExporterOptions ExportOptions( MoveTemp( FullPath ), Mesh, LightmapUV, CollisionMesh );
+	if (NewFormat)
+	{
+		Impl2->DoExport( ExportedMeshElement, ExportOptions );
+	}
+	else
+	{
+		Impl->DoExport( ExportedMeshElement, ExportOptions );
+	}
 
 	return ExportedMeshElement;
 }
@@ -208,14 +321,15 @@ TSharedPtr< IDatasmithMeshElement > FDatasmithMeshExporter::ExportToUObject( con
 bool FDatasmithMeshExporter::ExportToUObject( TSharedPtr< IDatasmithMeshElement >& MeshElement, const TCHAR* Filepath, FDatasmithMesh& Mesh, FDatasmithMesh* CollisionMesh, EDSExportLightmapUV LightmapUV )
 {
 	FString FullPath( GetMeshFilePath( Filepath, MeshElement->GetName() ) );
-	FDatasmithMeshExporterImpl::FDatasmithMeshExporterOptions ExportOptions( MoveTemp( FullPath ), Mesh, LightmapUV, CollisionMesh );
+	FDatasmithMeshExporterOptions ExportOptions( FullPath, Mesh, LightmapUV, CollisionMesh );
 
 	return Impl->DoExport( MeshElement, ExportOptions );
 }
 
+
 FString FDatasmithMeshExporter::GetLastError() const
 {
-	return Impl->LastError;
+	return Impl->LastError.IsEmpty() ? Impl2->LastError : Impl->LastError;
 }
 
 bool FDatasmithMeshExporterImpl::DoExport( TSharedPtr< IDatasmithMeshElement >& MeshElement, const FDatasmithMeshExporterOptions& ExportOptions )
@@ -278,7 +392,7 @@ void FDatasmithMeshExporterImpl::PostExport( const FDatasmithMesh& DatasmithMesh
 	MeshElement->SetLightmapSourceUV( DatasmithMesh.GetLightmapSourceUVChannel() );
 }
 
-void FDatasmithMeshExporterImpl::CreateDefaultUVs( FDatasmithMesh& Mesh )
+void CreateDefaultUVs( FDatasmithMesh& Mesh )
 {
 	if ( Mesh.GetUVChannelsCount() > 0 )
 	{
@@ -287,7 +401,7 @@ void FDatasmithMeshExporterImpl::CreateDefaultUVs( FDatasmithMesh& Mesh )
 
 	// Get the mesh description to generate BoxUV.
 	FMeshDescription MeshDescription;
-	RegisterStaticMeshAttributes(MeshDescription);
+	FStaticMeshAttributes(MeshDescription).Register();
 	FDatasmithMeshUtils::ToMeshDescription(Mesh, MeshDescription);
 	FUVMapParameters UVParameters(Mesh.GetExtents().GetCenter(), FQuat::Identity, Mesh.GetExtents().GetSize(), FVector::OneVector, FVector2D::UnitVector);
 	TMap<FVertexInstanceID, FVector2D> TexCoords;
@@ -326,11 +440,6 @@ void FDatasmithMeshExporterImpl::CreateDefaultUVs( FDatasmithMesh& Mesh )
 	}
 }
 
-void FDatasmithMeshExporterImpl::RegisterStaticMeshAttributes(FMeshDescription& MeshDescription)
-{
-	FStaticMeshAttributes(MeshDescription).Register();
-}
-
 TSharedPtr<UDatasmithMesh> FDatasmithMeshExporterImpl::GeneratePooledUDatasmithMeshFromFDatasmithMesh(const FDatasmithMesh& Mesh, bool bIsCollisionMesh)
 {
 	TSharedPtr<UDatasmithMesh> PooledMesh = GetPooledUDatasmithMesh(bIsCollisionMesh);
diff --git a/Engine/Source/Developer/Datasmith/DatasmithExporter/Private/DatasmithSceneExporter.cpp b/Engine/Source/Developer/Datasmith/DatasmithExporter/Private/DatasmithSceneExporter.cpp
--- a/Engine/Source/Developer/Datasmith/DatasmithExporter/Private/DatasmithSceneExporter.cpp
+++ b/Engine/Source/Developer/Datasmith/DatasmithExporter/Private/DatasmithSceneExporter.cpp
@@ -46,11 +46,6 @@ public:
 	FString Name;
 	FString OutputPath;
 	FString AssetsOutputPath;
-	FString Host;
-	FString Vendor;
-	FString ProductName;
-	FString ProductVersion;
-	FString Renderer;
 
 	uint64 ExportStartCycles;
 
@@ -90,7 +85,7 @@ void FDatasmithSceneExporterImpl::UpdateTextureElements( TSharedRef< IDatasmithS
 		}
 
 		FString& NewFilename = HashFilePathMap.FindOrAdd(TextureElement->GetFileHash());
-		
+
 		// If this texture has not been exported yet, find a unique name for it and copy its file to the asset output path.
 		if (NewFilename.IsEmpty())
 		{
@@ -217,13 +212,10 @@ FString FDatasmithSceneExporterImpl::GetFileNameWithHash(const FString& FullPath
 FDatasmithSceneExporter::FDatasmithSceneExporter()
 	: Impl( MakeUnique< FDatasmithSceneExporterImpl >() )
 {
-	Reset();
 }
 
-FDatasmithSceneExporter::~FDatasmithSceneExporter()
-{
-	Reset();
-}
+FDatasmithSceneExporter::~FDatasmithSceneExporter() = default;
+
 
 void FDatasmithSceneExporter::PreExport()
 {
@@ -322,12 +314,7 @@ void FDatasmithSceneExporter::Export( TSharedRef< IDatasmithScene > DatasmithSce
 
 void FDatasmithSceneExporter::Reset()
 {
-	Impl->Host = TEXT("");
-	Impl->Renderer = TEXT("");
-	Impl->ProgressManager = nullptr;
-	Impl->Logger = nullptr;
-
-	Impl->ExportStartCycles = 0;
+	Impl = MakeUnique<FDatasmithSceneExporterImpl>();
 }
 
 void FDatasmithSceneExporter::SetProgressManager( const TSharedPtr< IDatasmithProgressManager >& InProgressManager )
diff --git a/Engine/Source/Developer/Datasmith/DatasmithExporter/Public/DatasmithMeshExporter.h b/Engine/Source/Developer/Datasmith/DatasmithExporter/Public/DatasmithMeshExporter.h
--- a/Engine/Source/Developer/Datasmith/DatasmithExporter/Public/DatasmithMeshExporter.h
+++ b/Engine/Source/Developer/Datasmith/DatasmithExporter/Public/DatasmithMeshExporter.h
@@ -5,7 +5,6 @@
 #include "Templates/SharedPointer.h"
 
 class FDatasmithMesh;
-class FDatasmithMeshExporterImpl;
 class IDatasmithMeshElement;
 
 class DATASMITHEXPORTER_API FDatasmithMeshExporter
@@ -25,8 +24,8 @@ public:
 	 *
 	 * @return				A IDatasmithMeshElement that refers to the exported file
 	 */
-	TSharedPtr< IDatasmithMeshElement > ExportToUObject( const TCHAR* Filepath, const TCHAR* Filename, FDatasmithMesh& Mesh, FDatasmithMesh* CollisionMesh, EDSExportLightmapUV LightmapUV );
-	
+	TSharedPtr< IDatasmithMeshElement > ExportToUObject( const TCHAR* Filepath, const TCHAR* Filename, FDatasmithMesh& Mesh, FDatasmithMesh* CollisionMesh, EDSExportLightmapUV LightmapUV, bool NewFormat=false );
+
 	/**
 	 * Exports a FDatasmithMesh as a UObject and link it to the given IDatasmithMeshElementElement.
 	 *
@@ -46,6 +45,6 @@ public:
 	FString GetLastError() const;
 
 private:
-
-	FDatasmithMeshExporterImpl* Impl;
+	TUniquePtr<class FDatasmithMeshExporterImpl> Impl;
+	TUniquePtr<class FDatasmithMeshExporterImpl2> Impl2;
 };
diff --git a/Engine/Source/Runtime/Datasmith/DatasmithCore/Private/DatasmithMeshUObject.cpp b/Engine/Source/Runtime/Datasmith/DatasmithCore/Private/DatasmithMeshUObject.cpp
--- a/Engine/Source/Runtime/Datasmith/DatasmithCore/Private/DatasmithMeshUObject.cpp
+++ b/Engine/Source/Runtime/Datasmith/DatasmithCore/Private/DatasmithMeshUObject.cpp
@@ -1,6 +1,8 @@
 // Copyright Epic Games, Inc. All Rights Reserved.
 
 #include "DatasmithMeshUObject.h"
+#include "Serialization/MemoryReader.h"
+#include "Serialization/MemoryWriter.h"
 
 void FDatasmithMeshSourceModel::SerializeBulkData(FArchive& Ar, UObject* Owner)
 {
@@ -16,3 +18,53 @@ void UDatasmithMesh::Serialize(FArchive& Ar)
 		SourceModel.SerializeBulkData( Ar, this );
 	}
 }
+
+void operator<<(FArchive& Ar, FDatasmithMeshModels& Models)
+{
+	Ar << Models.MeshName;
+	Ar << Models.bIsCollisionMesh;
+	Ar << Models.SourceModels;
+}
+
+void operator<<(FArchive& Ar, FDatasmithPackedMeshes& Pack)
+{
+	FString Guard = Ar.IsLoading() ? TEXT("") : TEXT("FDatasmithPackedMeshes");
+	Ar << Guard;
+	if (!ensure(Guard == TEXT("FDatasmithPackedMeshes")))
+	{
+		Ar.SetError();
+		return;
+	}
+
+	uint8 SerialVersionMajor = 1; // increment this when forward compatibility is not possible
+	Ar << SerialVersionMajor;
+	if (SerialVersionMajor > 1)
+	{
+		Ar.SetError();
+		return;
+	}
+
+	uint8 SerialVersionMinor = 0; // increment this to handle logic changes that preserves forward compatibility
+	Ar << SerialVersionMinor;
+
+	uint8 BufferType = 0;
+	Ar << BufferType; // (MeshDesc, Zipped Mesh desc;...)
+
+	if (Ar.IsLoading())
+	{
+		TArray<uint8> Bytes;
+		Ar << Bytes;
+		FMemoryReader Buffer(Bytes, true);
+		Buffer << Pack.MeshesToExport;
+	}
+	else
+	{
+		TArray<uint8> Bytes;
+		FMemoryWriter Buffer(Bytes, true);
+		Buffer << Pack.MeshesToExport;
+		Ar << Bytes;
+		FMD5 Md5;
+		Md5.Update(Bytes.GetData(), Bytes.Num());
+		Pack.OutHash.Set(Md5);
+	}
+}
diff --git a/Engine/Source/Runtime/Datasmith/DatasmithCore/Private/DatasmithUtils.cpp b/Engine/Source/Runtime/Datasmith/DatasmithCore/Private/DatasmithUtils.cpp
--- a/Engine/Source/Runtime/Datasmith/DatasmithCore/Private/DatasmithUtils.cpp
+++ b/Engine/Source/Runtime/Datasmith/DatasmithCore/Private/DatasmithUtils.cpp
@@ -293,25 +293,17 @@ bool FDatasmithMeshUtils::ToRawMesh(const FDatasmithMesh& Mesh, FRawMesh& RawMes
 	return true;
 }
 
-bool FDatasmithMeshUtils::ToMeshDescription(FDatasmithMesh& DsMesh, FMeshDescription& MeshDescription)
+bool FDatasmithMeshUtils::ToMeshDescription(const FDatasmithMesh& DsMesh, FMeshDescription& MeshDescription)
 {
 	MeshDescription.Empty();
 
 	FStaticMeshAttributes Attributes(MeshDescription);
+	Attributes.Register();
 	TVertexAttributesRef<FVector3f> VertexPositions = Attributes.GetVertexPositions();
 	TVertexInstanceAttributesRef<FVector3f> VertexInstanceNormals = Attributes.GetVertexInstanceNormals();
 	TVertexInstanceAttributesRef<FVector2f> VertexInstanceUVs = Attributes.GetVertexInstanceUVs();
 	TPolygonGroupAttributesRef<FName> PolygonGroupImportedMaterialSlotNames = Attributes.GetPolygonGroupMaterialSlotNames();
 
-	// Prepared for static mesh usage ?
-	if (!ensure(VertexPositions.IsValid())
-	 || !ensure(VertexInstanceNormals.IsValid())
-	 || !ensure(VertexInstanceUVs.IsValid())
-	 || !ensure(PolygonGroupImportedMaterialSlotNames.IsValid()) )
-	{
-		return false;
-	}
-
 	// Reserve space for attributes.
 	int32 VertexCount = DsMesh.GetVerticesCount();
 	int32 TriangleCount = DsMesh.GetFacesCount();
@@ -1219,7 +1211,10 @@ namespace DatasmithSceneUtilsImpl
 			TFunction<void(TSharedPtr<IDatasmithElement>&&, const FString&)> AddAsset;
 			AddAsset = [this](TSharedPtr<IDatasmithElement>&& InElementPtr, const FString& AssetPrefix) -> void
 			{
-				AssetElementMapping.Add(AssetPrefix + InElementPtr->GetName(), MoveTemp(InElementPtr));
+				if (InElementPtr)
+				{
+					AssetElementMapping.Add(AssetPrefix + InElementPtr->GetName(), MoveTemp(InElementPtr));
+				}
 			};
 
 			for (int32 Index = 0; Index < Scene->GetTexturesCount(); ++Index)
@@ -1527,7 +1522,6 @@ namespace DatasmithSceneUtilsImpl
 void FDatasmithSceneUtils::CleanUpScene(TSharedRef<IDatasmithScene> Scene, bool bRemoveUnused)
 {
 	using namespace DatasmithSceneUtilsImpl;
-	using namespace DirectLink;
 
 	TRACE_CPUPROFILER_EVENT_SCOPE(FDatasmithSceneUtils::CleanUpScene);
 
diff --git a/Engine/Source/Runtime/Datasmith/DatasmithCore/Public/DatasmithMeshUObject.h b/Engine/Source/Runtime/Datasmith/DatasmithCore/Public/DatasmithMeshUObject.h
--- a/Engine/Source/Runtime/Datasmith/DatasmithCore/Public/DatasmithMeshUObject.h
+++ b/Engine/Source/Runtime/Datasmith/DatasmithCore/Public/DatasmithMeshUObject.h
@@ -4,6 +4,7 @@
 
 #include "CoreMinimal.h"
 
+#include "MeshDescription.h"
 #include "RawMesh.h"
 
 #include "UObject/Object.h"
@@ -42,3 +43,19 @@ public:
 	UPROPERTY()
 	TArray< FDatasmithMeshSourceModel > SourceModels;
 };
+
+struct FDatasmithMeshModels
+{
+	FString MeshName;
+	bool bIsCollisionMesh;
+	TArray<FMeshDescription> SourceModels;
+
+	DATASMITHCORE_API friend void operator << (FArchive& Ar, FDatasmithMeshModels& Models);
+};
+
+struct FDatasmithPackedMeshes
+{
+	TArray<FDatasmithMeshModels> MeshesToExport;
+	DATASMITHCORE_API friend void operator << (FArchive& Ar, FDatasmithPackedMeshes& Pack);
+	FMD5Hash OutHash;
+};
\ No newline at end of file
diff --git a/Engine/Source/Runtime/Datasmith/DatasmithCore/Public/DatasmithUtils.h b/Engine/Source/Runtime/Datasmith/DatasmithCore/Public/DatasmithUtils.h
--- a/Engine/Source/Runtime/Datasmith/DatasmithCore/Public/DatasmithUtils.h
+++ b/Engine/Source/Runtime/Datasmith/DatasmithCore/Public/DatasmithUtils.h
@@ -135,7 +135,7 @@ public:
 	 * For example a collision mesh don't need to have a valid raw mesh
 	 */
 	static bool ToRawMesh(const FDatasmithMesh& Mesh, FRawMesh& RawMesh, bool bValidateRawMesh = true);
-	static bool ToMeshDescription(FDatasmithMesh& DsMesh, FMeshDescription& MeshDescription);
+	static bool ToMeshDescription(const FDatasmithMesh& DsMesh, FMeshDescription& MeshDescription);
 
 
 	/**
