Datasmith: MeshSerialization refactor

diff --git a/Engine/Plugins/Enterprise/DatasmithImporter/Source/DatasmithNativeTranslator/Private/DatasmithNativeTranslator.cpp b/Engine/Plugins/Enterprise/DatasmithImporter/Source/DatasmithNativeTranslator/Private/DatasmithNativeTranslator.cpp
--- a/Engine/Plugins/Enterprise/DatasmithImporter/Source/DatasmithNativeTranslator/Private/DatasmithNativeTranslator.cpp
+++ b/Engine/Plugins/Enterprise/DatasmithImporter/Source/DatasmithNativeTranslator/Private/DatasmithNativeTranslator.cpp
@@ -4,22 +4,14 @@
 
 #include "DatasmithAnimationElements.h"
 #include "DatasmithAnimationSerializer.h"
-#include "DatasmithMeshUObject.h"
+#include "DatasmithMeshSerialization.h"
 #include "DatasmithSceneSource.h"
 #include "DatasmithSceneXmlReader.h"
 #include "IDatasmithSceneElements.h"
-#include "Utility/DatasmithMeshHelper.h"
 
-#include "HAL/FileManager.h"
-#include "Serialization/MemoryReader.h"
-#include "StaticMeshOperations.h"
-#include "Templates/SharedPointer.h"
-#include "Templates/UniquePtr.h"
 #include "Misc/Paths.h"
-#include "UObject/StrongObjectPtr.h"
-
+#include "HAL/FileManager.h"
 
-DEFINE_LOG_CATEGORY(LogDatasmithNativeTranslator);
 
 FString FDatasmithNativeTranslator::ResolveFilePath(const FString& FilePath, const TArray<FString>& ResourcePaths)
 {
@@ -119,139 +111,10 @@ bool FDatasmithNativeTranslator::LoadScene(TSharedRef<IDatasmithScene> OutScene)
 	return bParsingResult;
 }
 
-namespace DatasmithNativeTranslatorImpl
-{
-	TOptional<FMeshDescription> ExtractToMeshDescription(FDatasmithMeshSourceModel& SourceModel)
-	{
-		FRawMesh RawMesh;
-		SourceModel.RawMeshBulkData.LoadRawMesh( RawMesh );
-
-		if ( !RawMesh.IsValid() )
-		{
-			return {};
-		}
-
-		// RawMesh -> MeshDescription conversion requires an {mat_index: slot_name} map for its PolygonGroups.
-		TMap<int32, FName> GroupNamePerGroupIndex;
-
-		// There is no guaranty that incoming RawMesh.FaceMaterialIndices are sequential, but the conversion assumes so.
-		// -> we remap materials identifiers to material indices
-		// eg:
-		//   incoming per-face mat identifier   5   5   1   1   1   99   99
-		//   remapped per-face index            0   0   1   1   1   2    2
-		//   per PolygonGroup FName:           "5" "5" "1" "1" "1" "99" "99"
-		TSet<int32> MaterialIdentifiers;
-		for (int32& MatIdentifier : RawMesh.FaceMaterialIndices)
-		{
-			bool bAlreadyIn = false;
-			int32 IndexOfIdentifier = MaterialIdentifiers.Add(MatIdentifier, &bAlreadyIn).AsInteger();
-
-			// identifier -> name association
-			if (!bAlreadyIn)
-			{
-				FName MaterialSlotName = DatasmithMeshHelper::DefaultSlotName(MatIdentifier);
-				GroupNamePerGroupIndex.Add(IndexOfIdentifier, MaterialSlotName);
-			}
-
-			// remap old identifier to material index
-			MatIdentifier = IndexOfIdentifier;
-		}
-
-		FMeshDescription MeshDescription;
-		DatasmithMeshHelper::PrepareAttributeForStaticMesh(MeshDescription);
-
-		// Do not compute normals and tangents during conversion since we have other operations to apply
-		// on the mesh that might invalidate them anyway and we must also validate the mesh to detect
-		// vertex positions containing NaN before doing computation as MikkTSpace crashes on NaN values.
-		const bool bSkipNormalsAndTangents = true;
-		FStaticMeshOperations::ConvertFromRawMesh(RawMesh, MeshDescription, GroupNamePerGroupIndex, bSkipNormalsAndTangents);
-		return MeshDescription;
-	}
-
-
-	TArray< FDatasmithMeshModels > GetDatasmithMeshFromMeshPath( const FString& MeshPath )
-	{
-		TArray< FDatasmithMeshModels > Result;
-
-		TUniquePtr<FArchive> Archive( IFileManager::Get().CreateFileReader(*MeshPath) );
-		if ( !Archive.IsValid() )
-		{
-			UE_LOG(LogDatasmithNativeTranslator, Warning, TEXT("Cannot read file %s"), *MeshPath);
-			return Result;
-		}
-
-		int32 LeagacyNumMeshesCount = 0;
-		*Archive << LeagacyNumMeshesCount;
-
-		if (LeagacyNumMeshesCount == 0)
-		{
-			FDatasmithPackedMeshes Pack;
-			*Archive << Pack;
-
-			if (!Archive->IsError())
-			{
-				for (FDatasmithMeshModels& Mesh : Pack.MeshesToExport)
-				{
-					FDatasmithMeshModels& MeshInternal = Result.AddDefaulted_GetRef();
-					MeshInternal.bIsCollisionMesh = Mesh.bIsCollisionMesh;
-					MeshInternal.SourceModels = MoveTemp(Mesh.SourceModels);
-				}
-			}
-			else
-			{
-				UE_LOG(LogDatasmithNativeTranslator, Warning, TEXT("Failed to read meshes from %s"), *MeshPath);
-			}
-		}
-		else
-		{
-			UDatasmithMesh* DatasmithMesh = nullptr;
-			{
-				// Make sure the new UDatasmithMesh object is not created while a garbage collection is performed
-				FGCScopeGuard GCGuard;
-				// Setting the RF_Standalone bitmask on the new UDatasmithMesh object, to make sure it is not garbage collected
-				// while loading and processing the udsmesh file. This can happen on very big meshes (5M+ triangles)
-				DatasmithMesh = NewObject< UDatasmithMesh >(GetTransientPackage(), NAME_None, RF_Standalone);
-			}
-
-			// Currently we only have 1 mesh per file. If there's a second mesh, it will be a CollisionMesh
-			while ( LeagacyNumMeshesCount-- > 0)
-			{
-				TArray< uint8 > Bytes;
-				*Archive << Bytes;
-
-				FMemoryReader MemoryReader( Bytes, true );
-				MemoryReader.ArIgnoreClassRef = false;
-				MemoryReader.ArIgnoreArchetypeRef = false;
-				MemoryReader.SetWantBinaryPropertySerialization(true);
-				DatasmithMesh->Serialize( MemoryReader );
-
-				FDatasmithMeshModels& MeshInternal = Result.AddDefaulted_GetRef();
-				MeshInternal.bIsCollisionMesh = DatasmithMesh->bIsCollisionMesh;
-				for (FDatasmithMeshSourceModel& SourceModel : DatasmithMesh->SourceModels)
-				{
-					if (TOptional<FMeshDescription> OptionalMesh = ExtractToMeshDescription(SourceModel))
-					{
-						MeshInternal.SourceModels.Add(MoveTemp(*OptionalMesh));
-					}
-				}
-			}
-
-			// Tell the garbage collector DatasmithMesh can now be deleted.
-			DatasmithMesh->ClearInternalFlags(EInternalObjectFlags::Async);
-			DatasmithMesh->ClearFlags(RF_Standalone);
-		}
-
-		return Result;
-	}
-} // ns DatasmithNativeTranslatorImpl
-
-
 bool FDatasmithNativeTranslator::LoadStaticMesh(const TSharedRef<IDatasmithMeshElement> MeshElement, FDatasmithMeshElementPayload& OutMeshPayload)
 {
 	TRACE_CPUPROFILER_EVENT_SCOPE(FDatasmithNativeTranslator::LoadStaticMesh);
 
-	using namespace DatasmithNativeTranslatorImpl;
-
 	FString FilePath = MeshElement->GetFile();
 
 	for (FDatasmithMeshModels& DatasmithMesh : GetDatasmithMeshFromMeshPath( FilePath ))
@@ -290,3 +153,4 @@ bool FDatasmithNativeTranslator::LoadLevelSequence(const TSharedRef<IDatasmithLe
 	}
 	return false;
 }
+
diff --git a/Engine/Plugins/Enterprise/DatasmithImporter/Source/DatasmithNativeTranslator/Public/DatasmithNativeTranslator.h b/Engine/Plugins/Enterprise/DatasmithImporter/Source/DatasmithNativeTranslator/Public/DatasmithNativeTranslator.h
--- a/Engine/Plugins/Enterprise/DatasmithImporter/Source/DatasmithNativeTranslator/Public/DatasmithNativeTranslator.h
+++ b/Engine/Plugins/Enterprise/DatasmithImporter/Source/DatasmithNativeTranslator/Public/DatasmithNativeTranslator.h
@@ -5,8 +5,6 @@
 #include "CoreMinimal.h"
 #include "DatasmithTranslator.h"
 
-#include "Logging/LogMacros.h"
-DECLARE_LOG_CATEGORY_EXTERN(LogDatasmithNativeTranslator, Log, All);
 
 class DATASMITHNATIVETRANSLATOR_API FDatasmithNativeTranslator : public IDatasmithTranslator
 {
diff --git a/Engine/Plugins/Enterprise/DatasmithImporter/Source/DatasmithTranslator/Private/Utility/DatasmithMeshHelper.cpp b/Engine/Plugins/Enterprise/DatasmithImporter/Source/DatasmithTranslator/Private/Utility/DatasmithMeshHelper.cpp
--- a/Engine/Plugins/Enterprise/DatasmithImporter/Source/DatasmithTranslator/Private/Utility/DatasmithMeshHelper.cpp
+++ b/Engine/Plugins/Enterprise/DatasmithImporter/Source/DatasmithTranslator/Private/Utility/DatasmithMeshHelper.cpp
@@ -268,49 +268,7 @@ namespace DatasmithMeshHelper
 
 	void CreateDefaultUVs(FDatasmithMesh& DatasmithMesh)
 	{
-		// Get the mesh description to generate BoxUV.
-		FMeshDescription MeshDescription;
-		PrepareAttributeForStaticMesh(MeshDescription);
-		FDatasmithMeshUtils::ToMeshDescription(DatasmithMesh, MeshDescription);
-		FUVMapParameters UVParameters(DatasmithMesh.GetExtents().GetCenter(), FQuat::Identity, DatasmithMesh.GetExtents().GetSize(), FVector::OneVector, FVector2D::UnitVector);
-		TMap<FVertexInstanceID, FVector2D> TexCoords;
-		FStaticMeshOperations::GenerateBoxUV(MeshDescription, UVParameters, TexCoords);
-
-		// Put the results in a map to determine the number of unique values.
-		TMap<FVector2D, TArray<int32>> UniqueTexCoordMap;
-		for (const TPair<FVertexInstanceID, FVector2D>& Pair : TexCoords)
-		{
-			TArray<int32>& MappedIndices = UniqueTexCoordMap.FindOrAdd(Pair.Value);
-			MappedIndices.Add(Pair.Key.GetValue());
-		}
-
-		//Set the UV values
-		if (DatasmithMesh.GetUVChannelsCount() == 0)
-		{
-			DatasmithMesh.AddUVChannel();
-		}
-		DatasmithMesh.SetUVCount(0, UniqueTexCoordMap.Num());
-		int32 UVIndex = 0;
-		TArray<int32> IndicesMapping;
-		IndicesMapping.AddZeroed(TexCoords.Num());
-
-		for (const TPair<FVector2D, TArray<int32>>& UniqueCoordPair : UniqueTexCoordMap)
-		{
-			DatasmithMesh.SetUV(0, UVIndex, UniqueCoordPair.Key.X, UniqueCoordPair.Key.Y);
-			for (int32 IndicesIndex : UniqueCoordPair.Value)
-			{
-				IndicesMapping[IndicesIndex] = UVIndex;
-			}
-			UVIndex++;
-		}
-
-		//Map the UV indices.
-		for (int32 FaceIndex = 0; FaceIndex < DatasmithMesh.GetFacesCount(); ++FaceIndex)
-		{
-			const int32 IndicesOffset = FaceIndex * 3;
-			check(IndicesOffset + 2 < IndicesMapping.Num());
-			DatasmithMesh.SetFaceUV(FaceIndex, 0, IndicesMapping[IndicesOffset + 0], IndicesMapping[IndicesOffset + 1], IndicesMapping[IndicesOffset + 2]);
-		}
+		FDatasmithMeshUtils::CreateDefaultUVsWithLOD(DatasmithMesh);
 	}
 
 	void CreateDefaultUVs(FMeshDescription& MeshDescription)
diff --git a/Engine/Source/Developer/Datasmith/DatasmithExporter/Private/DatasmithMeshExporter.cpp b/Engine/Source/Developer/Datasmith/DatasmithExporter/Private/DatasmithMeshExporter.cpp
--- a/Engine/Source/Developer/Datasmith/DatasmithExporter/Private/DatasmithMeshExporter.cpp
+++ b/Engine/Source/Developer/Datasmith/DatasmithExporter/Private/DatasmithMeshExporter.cpp
@@ -6,18 +6,15 @@
 #include "DatasmithExporterManager.h"
 #include "DatasmithMesh.h"
 #include "DatasmithMeshUObject.h"
+#include "DatasmithMeshSerialization.h"
 #include "DatasmithSceneFactory.h"
 #include "DatasmithUtils.h"
 
 #include "Containers/LockFreeList.h"
 #include "HAL/FileManager.h"
-#include "HAL/PlatformFileManager.h"
 #include "Misc/Guid.h"
 #include "Misc/Paths.h"
 #include "Serialization/MemoryWriter.h"
-#include "StaticMeshAttributes.h"
-#include "StaticMeshOperations.h"
-#include "UVMapSettings.h"
 
 /**
  *	Number of datsmith meshes pending deletion for which the Garbage Collection will be triggered.
@@ -38,21 +35,20 @@ struct FDatasmithMeshExporterOptions
 	EDSExportLightmapUV LightmapUV;
 	FDatasmithMesh* CollisionMesh;
 };
-void CreateDefaultUVs(FDatasmithMesh& DatasmithMesh);
 
 /**
  *	Implementation class of the DatasmithMeshExporter
  *	We use a lockfree UDatasmithMesh pool to avoid creating new UObject when exporting and reduces our memory footprint.
  */
-class FDatasmithMeshExporterImpl
+class FDatasmithMeshExporterLegacyImpl
 {
 public:
 
-	FDatasmithMeshExporterImpl() :
+	FDatasmithMeshExporterLegacyImpl() :
 		UniqueID(FGuid::NewGuid())
 	{}
 
-	~FDatasmithMeshExporterImpl()
+	~FDatasmithMeshExporterLegacyImpl()
 	{
 		ClearUDatasmithMeshPool();
 	}
@@ -98,7 +94,7 @@ private:
 	 */
 	struct FPooledUDatasmithMeshDeleter
 	{
-		FPooledUDatasmithMeshDeleter(FDatasmithMeshExporterImpl* InExporterPtr)
+		FPooledUDatasmithMeshDeleter(FDatasmithMeshExporterLegacyImpl* InExporterPtr)
 			: ExporterPtr(InExporterPtr)
 		{}
 
@@ -108,14 +104,14 @@ private:
 			ExporterPtr->ReturnUDatasmithMeshToPool(Object);
 		}
 
-		FDatasmithMeshExporterImpl* ExporterPtr;
+		FDatasmithMeshExporterLegacyImpl* ExporterPtr;
 	};
 };
 
-TAtomic<int32> FDatasmithMeshExporterImpl::NumberOfUMeshPendingGC(0);
+TAtomic<int32> FDatasmithMeshExporterLegacyImpl::NumberOfUMeshPendingGC(0);
 
 
-class FDatasmithMeshExporterImpl2
+class FDatasmithMeshExporterImpl
 {
 public:
 	bool DoExport( TSharedPtr< IDatasmithMeshElement >& MeshElement, const FDatasmithMeshExporterOptions& ExportOptions );
@@ -125,7 +121,7 @@ private:
 	bool WriteMeshFile(const FDatasmithMeshExporterOptions& ExporterOptions, FMD5Hash& OutHash);
 };
 
-bool FDatasmithMeshExporterImpl2::DoExport(TSharedPtr< IDatasmithMeshElement >& MeshElement, const FDatasmithMeshExporterOptions& ExportOptions)
+bool FDatasmithMeshExporterImpl::DoExport(TSharedPtr< IDatasmithMeshElement >& MeshElement, const FDatasmithMeshExporterOptions& ExportOptions)
 {
 	FDatasmithMesh& Mesh = ExportOptions.Mesh;
 
@@ -136,14 +132,7 @@ bool FDatasmithMeshExporterImpl2::DoExport(TSharedPtr< IDatasmithMeshElement >&
 	}
 
 	// make sure we have at least one UV channels on each source models
-	CreateDefaultUVs( Mesh );
-	for ( int32 LODIndex = 0; LODIndex < Mesh.GetLODsCount(); ++LODIndex )
-	{
-		if ( FDatasmithMesh* LODMesh = Mesh.GetLOD( LODIndex ) )
-		{
-			CreateDefaultUVs( *LODMesh );
-		}
-	}
+	FDatasmithMeshUtils::CreateDefaultUVsWithLOD( Mesh );
 
 	FMD5Hash Hash;
 	if ( WriteMeshFile( ExportOptions, Hash ) )
@@ -175,7 +164,7 @@ bool FDatasmithMeshExporterImpl2::DoExport(TSharedPtr< IDatasmithMeshElement >&
 }
 
 
-bool FDatasmithMeshExporterImpl::ExportMeshes( const FDatasmithMeshExporterOptions& ExporterOptions, FMD5Hash& OutHash )
+bool FDatasmithMeshExporterLegacyImpl::ExportMeshes( const FDatasmithMeshExporterOptions& ExporterOptions, FMD5Hash& OutHash )
 {
 	TArray<TSharedPtr<UDatasmithMesh>, TInlineAllocator<2>> MeshesToExport;
 
@@ -190,7 +179,6 @@ bool FDatasmithMeshExporterImpl::ExportMeshes( const FDatasmithMeshExporterOptio
 		MeshesToExport.Add( GeneratePooledUDatasmithMeshFromFDatasmithMesh( *ExporterOptions.CollisionMesh, bIsCollisionMesh ) );
 	}
 
-	IPlatformFile& PlatformFile = FPlatformFileManager::Get().GetPlatformFile();
 	TUniquePtr<FArchive> Archive( IFileManager::Get().CreateFileWriter( *ExporterOptions.MeshFullPath ) );
 
 	if ( !Archive.IsValid() )
@@ -233,8 +221,8 @@ bool FDatasmithMeshExporterImpl::ExportMeshes( const FDatasmithMeshExporterOptio
 
 
 FDatasmithMeshExporter::FDatasmithMeshExporter()
-	: Impl(MakeUnique<FDatasmithMeshExporterImpl>())
-	, Impl2(MakeUnique<FDatasmithMeshExporterImpl2>())
+	: LegacyImpl(MakeUnique<FDatasmithMeshExporterLegacyImpl>())
+	, Impl(MakeUnique<FDatasmithMeshExporterImpl>())
 {}
 
 FDatasmithMeshExporter::~FDatasmithMeshExporter() = default;
@@ -251,7 +239,7 @@ FString GetMeshFilePath( const TCHAR* Filepath, const TCHAR* Filename )
 }
 
 
-bool FDatasmithMeshExporterImpl2::WriteMeshFile( const FDatasmithMeshExporterOptions& ExporterOptions, FMD5Hash& OutHash )
+bool FDatasmithMeshExporterImpl::WriteMeshFile( const FDatasmithMeshExporterOptions& ExporterOptions, FMD5Hash& OutHash )
 {
 	FDatasmithPackedMeshes Pack;
 
@@ -308,11 +296,11 @@ TSharedPtr< IDatasmithMeshElement > FDatasmithMeshExporter::ExportToUObject(cons
 	FDatasmithMeshExporterOptions ExportOptions( MoveTemp( FullPath ), Mesh, LightmapUV, CollisionMesh );
 	if (NewFormat)
 	{
-		Impl2->DoExport( ExportedMeshElement, ExportOptions );
+		Impl->DoExport( ExportedMeshElement, ExportOptions );
 	}
 	else
 	{
-		Impl->DoExport( ExportedMeshElement, ExportOptions );
+		LegacyImpl->DoExport( ExportedMeshElement, ExportOptions );
 	}
 
 	return ExportedMeshElement;
@@ -323,16 +311,16 @@ bool FDatasmithMeshExporter::ExportToUObject( TSharedPtr< IDatasmithMeshElement
 	FString FullPath( GetMeshFilePath( Filepath, MeshElement->GetName() ) );
 	FDatasmithMeshExporterOptions ExportOptions( FullPath, Mesh, LightmapUV, CollisionMesh );
 
-	return Impl->DoExport( MeshElement, ExportOptions );
+	return LegacyImpl->DoExport( MeshElement, ExportOptions );
 }
 
 
 FString FDatasmithMeshExporter::GetLastError() const
 {
-	return Impl->LastError.IsEmpty() ? Impl2->LastError : Impl->LastError;
+	return LegacyImpl->LastError.IsEmpty() ? Impl->LastError : LegacyImpl->LastError;
 }
 
-bool FDatasmithMeshExporterImpl::DoExport( TSharedPtr< IDatasmithMeshElement >& MeshElement, const FDatasmithMeshExporterOptions& ExportOptions )
+bool FDatasmithMeshExporterLegacyImpl::DoExport( TSharedPtr< IDatasmithMeshElement >& MeshElement, const FDatasmithMeshExporterOptions& ExportOptions )
 {
 	FMD5Hash Hash;
 	PreExport( ExportOptions );
@@ -355,7 +343,7 @@ bool FDatasmithMeshExporterImpl::DoExport( TSharedPtr< IDatasmithMeshElement >&
 	return false;
 }
 
-void FDatasmithMeshExporterImpl::PreExport( const FDatasmithMeshExporterOptions& ExporterOptions )
+void FDatasmithMeshExporterLegacyImpl::PreExport( const FDatasmithMeshExporterOptions& ExporterOptions )
 {
 	FDatasmithMesh& Mesh = ExporterOptions.Mesh;
 
@@ -365,22 +353,10 @@ void FDatasmithMeshExporterImpl::PreExport( const FDatasmithMeshExporterOptions&
 		Mesh.SetName( *FPaths::GetBaseFilename( ExporterOptions.MeshFullPath ) );
 	}
 
-	bool bHasUVs = Mesh.GetUVChannelsCount() > 0;
-	if ( !bHasUVs )
-	{
-		CreateDefaultUVs( Mesh );
-	}
-
-	for ( int32 LODIndex = 0; LODIndex < Mesh.GetLODsCount(); ++LODIndex )
-	{
-		if ( FDatasmithMesh* LODMesh = Mesh.GetLOD( LODIndex ) )
-		{
-			CreateDefaultUVs( *LODMesh );
-		}
-	}
+	FDatasmithMeshUtils::CreateDefaultUVsWithLOD( Mesh );
 }
 
-void FDatasmithMeshExporterImpl::PostExport( const FDatasmithMesh& DatasmithMesh, TSharedPtr< IDatasmithMeshElement >& MeshElement )
+void FDatasmithMeshExporterLegacyImpl::PostExport( const FDatasmithMesh& DatasmithMesh, TSharedPtr< IDatasmithMeshElement >& MeshElement )
 {
 	FBox Extents = DatasmithMesh.GetExtents();
 	float Width = Extents.Max[0] - Extents.Min[0];
@@ -392,55 +368,7 @@ void FDatasmithMeshExporterImpl::PostExport( const FDatasmithMesh& DatasmithMesh
 	MeshElement->SetLightmapSourceUV( DatasmithMesh.GetLightmapSourceUVChannel() );
 }
 
-void CreateDefaultUVs( FDatasmithMesh& Mesh )
-{
-	if ( Mesh.GetUVChannelsCount() > 0 )
-	{
-		return;
-	}
-
-	// Get the mesh description to generate BoxUV.
-	FMeshDescription MeshDescription;
-	FStaticMeshAttributes(MeshDescription).Register();
-	FDatasmithMeshUtils::ToMeshDescription(Mesh, MeshDescription);
-	FUVMapParameters UVParameters(Mesh.GetExtents().GetCenter(), FQuat::Identity, Mesh.GetExtents().GetSize(), FVector::OneVector, FVector2D::UnitVector);
-	TMap<FVertexInstanceID, FVector2D> TexCoords;
-	FStaticMeshOperations::GenerateBoxUV(MeshDescription, UVParameters, TexCoords);
-
-	// Put the results in a map to determine the number of unique values.
-	TMap<FVector2D, TArray<int32>> UniqueTexCoordMap;
-	for (const TPair<FVertexInstanceID, FVector2D>& Pair : TexCoords)
-	{
-		TArray<int32>& MappedIndices = UniqueTexCoordMap.FindOrAdd(Pair.Value);
-		MappedIndices.Add(Pair.Key.GetValue());
-	}
-
-	//Set the UV values
-	Mesh.AddUVChannel();
-	Mesh.SetUVCount(0, UniqueTexCoordMap.Num());
-	int32 UVIndex = 0;
-	TArray<int32> IndicesMapping;
-	IndicesMapping.AddZeroed(TexCoords.Num());
-	for (const TPair<FVector2D, TArray<int32>>& UniqueCoordPair : UniqueTexCoordMap)
-	{
-		Mesh.SetUV(0, UVIndex, UniqueCoordPair.Key.X, UniqueCoordPair.Key.Y);
-		for (int32 IndicesIndex : UniqueCoordPair.Value)
-		{
-			IndicesMapping[IndicesIndex] = UVIndex;
-		}
-		UVIndex++;
-	}
-
-	//Map the UV indices.
-	for (int32 FaceIndex = 0; FaceIndex < Mesh.GetFacesCount(); ++FaceIndex)
-	{
-		const int32 IndicesOffset = FaceIndex * 3;
-		check(IndicesOffset + 2 < IndicesMapping.Num());
-		Mesh.SetFaceUV(FaceIndex, 0, IndicesMapping[IndicesOffset + 0], IndicesMapping[IndicesOffset + 1], IndicesMapping[IndicesOffset + 2]);
-	}
-}
-
-TSharedPtr<UDatasmithMesh> FDatasmithMeshExporterImpl::GeneratePooledUDatasmithMeshFromFDatasmithMesh(const FDatasmithMesh& Mesh, bool bIsCollisionMesh)
+TSharedPtr<UDatasmithMesh> FDatasmithMeshExporterLegacyImpl::GeneratePooledUDatasmithMeshFromFDatasmithMesh(const FDatasmithMesh& Mesh, bool bIsCollisionMesh)
 {
 	TSharedPtr<UDatasmithMesh> PooledMesh = GetPooledUDatasmithMesh(bIsCollisionMesh);
 	FillUDatasmithMeshFromFDatasmithMesh(PooledMesh, Mesh, !bIsCollisionMesh);
@@ -451,7 +379,7 @@ TSharedPtr<UDatasmithMesh> FDatasmithMeshExporterImpl::GeneratePooledUDatasmithM
 /**
  * This function allows reusing an instanced UDatasmithMesh. Reusing the same object will avoid creating new garbage in memory.
  */
-void FDatasmithMeshExporterImpl::FillUDatasmithMeshFromFDatasmithMesh(TSharedPtr<UDatasmithMesh>& UMesh, const FDatasmithMesh& Mesh, bool bValidateRawMesh)
+void FDatasmithMeshExporterLegacyImpl::FillUDatasmithMeshFromFDatasmithMesh(TSharedPtr<UDatasmithMesh>& UMesh, const FDatasmithMesh& Mesh, bool bValidateRawMesh)
 {
 	UMesh->MeshName = Mesh.GetName();
 
@@ -480,7 +408,7 @@ void FDatasmithMeshExporterImpl::FillUDatasmithMeshFromFDatasmithMesh(TSharedPtr
 	}
 }
 
-TSharedPtr<UDatasmithMesh> FDatasmithMeshExporterImpl::GetPooledUDatasmithMesh(bool bIsCollisionMesh)
+TSharedPtr<UDatasmithMesh> FDatasmithMeshExporterLegacyImpl::GetPooledUDatasmithMesh(bool bIsCollisionMesh)
 {
 	const FString GuidString = UniqueID.ToString();
 	const FString CollisionString = bIsCollisionMesh ? TEXT("_Collision") : TEXT("");
@@ -499,7 +427,7 @@ TSharedPtr<UDatasmithMesh> FDatasmithMeshExporterImpl::GetPooledUDatasmithMesh(b
 	}
 }
 
-void FDatasmithMeshExporterImpl::ReturnUDatasmithMeshToPool(UDatasmithMesh*& UMesh)
+void FDatasmithMeshExporterLegacyImpl::ReturnUDatasmithMeshToPool(UDatasmithMesh*& UMesh)
 {
 	//Clear the UDatasmithMesh.
 	UMesh->SourceModels.Empty();
@@ -510,7 +438,7 @@ void FDatasmithMeshExporterImpl::ReturnUDatasmithMeshToPool(UDatasmithMesh*& UMe
 	UMesh = nullptr;
 }
 
-void FDatasmithMeshExporterImpl::ClearUDatasmithMeshPool()
+void FDatasmithMeshExporterLegacyImpl::ClearUDatasmithMeshPool()
 {
 	TArray<UDatasmithMesh*> PooledMeshes;
 	DatasmithMeshUObjectPool.PopAll(PooledMeshes);
diff --git a/Engine/Source/Developer/Datasmith/DatasmithExporter/Public/DatasmithMeshExporter.h b/Engine/Source/Developer/Datasmith/DatasmithExporter/Public/DatasmithMeshExporter.h
--- a/Engine/Source/Developer/Datasmith/DatasmithExporter/Public/DatasmithMeshExporter.h
+++ b/Engine/Source/Developer/Datasmith/DatasmithExporter/Public/DatasmithMeshExporter.h
@@ -45,6 +45,6 @@ public:
 	FString GetLastError() const;
 
 private:
+	TUniquePtr<class FDatasmithMeshExporterLegacyImpl> LegacyImpl;
 	TUniquePtr<class FDatasmithMeshExporterImpl> Impl;
-	TUniquePtr<class FDatasmithMeshExporterImpl2> Impl2;
 };
diff --git a/Engine/Source/Runtime/Datasmith/DatasmithCore/Private/DatasmithMeshSerialization.cpp b/Engine/Source/Runtime/Datasmith/DatasmithCore/Private/DatasmithMeshSerialization.cpp
new file mode 100644
--- /dev/null
+++ b/Engine/Source/Runtime/Datasmith/DatasmithCore/Private/DatasmithMeshSerialization.cpp
@@ -0,0 +1,191 @@
+// Copyright Epic Games, Inc. All Rights Reserved.
+#include "DatasmithMeshSerialization.h"
+
+#include "DatasmithMeshUObject.h"
+
+#include "HAL/FileManager.h"
+#include "Serialization/MemoryReader.h"
+#include "Serialization/MemoryWriter.h"
+#include "StaticMeshAttributes.h"
+#include "StaticMeshOperations.h"
+
+void operator<<(FArchive& Ar, FDatasmithMeshModels& Models)
+{
+	Ar << Models.MeshName;
+	Ar << Models.bIsCollisionMesh;
+	Ar << Models.SourceModels;
+}
+
+void operator<<(FArchive& Ar, FDatasmithPackedMeshes& Pack)
+{
+	FString Guard = Ar.IsLoading() ? TEXT("") : TEXT("FDatasmithPackedMeshes");
+	Ar << Guard;
+	if (!ensure(Guard == TEXT("FDatasmithPackedMeshes")))
+	{
+		Ar.SetError();
+		return;
+	}
+
+	uint8 SerialVersionMajor = 1; // increment this when forward compatibility is not possible
+	Ar << SerialVersionMajor;
+	if (SerialVersionMajor > 1)
+	{
+		Ar.SetError();
+		return;
+	}
+
+	uint8 SerialVersionMinor = 0; // increment this to handle logic changes that preserves forward compatibility
+	Ar << SerialVersionMinor;
+
+	uint8 BufferType = 0;
+	Ar << BufferType; // (MeshDesc, Zipped Mesh desc;...)
+
+	if (Ar.IsLoading())
+	{
+		TArray<uint8> Bytes;
+		Ar << Bytes;
+		FMemoryReader Buffer(Bytes, true);
+		Buffer << Pack.MeshesToExport;
+	}
+	else
+	{
+		TArray<uint8> Bytes;
+		FMemoryWriter Buffer(Bytes, true);
+		Buffer << Pack.MeshesToExport;
+		Ar << Bytes;
+		FMD5 Md5;
+		Md5.Update(Bytes.GetData(), Bytes.Num());
+		Pack.OutHash.Set(Md5);
+	}
+}
+
+TOptional<FMeshDescription> ExtractToMeshDescription(FDatasmithMeshSourceModel& SourceModel)
+{
+	FRawMesh RawMesh;
+	SourceModel.RawMeshBulkData.LoadRawMesh( RawMesh );
+
+	if ( !RawMesh.IsValid() )
+	{
+		return {};
+	}
+
+	// RawMesh -> MeshDescription conversion requires an {mat_index: slot_name} map for its PolygonGroups.
+	TMap<int32, FName> GroupNamePerGroupIndex;
+
+	// There is no guaranty that incoming RawMesh.FaceMaterialIndices are sequential, but the conversion assumes so.
+	// -> we remap materials identifiers to material indices
+	// eg:
+	//   incoming per-face mat identifier   5   5   1   1   1   99   99
+	//   remapped per-face index            0   0   1   1   1   2    2
+	//   per PolygonGroup FName:           "5" "5" "1" "1" "1" "99" "99"
+	TSet<int32> MaterialIdentifiers;
+	for (int32& MatIdentifier : RawMesh.FaceMaterialIndices)
+	{
+		bool bAlreadyIn = false;
+		int32 IndexOfIdentifier = MaterialIdentifiers.Add(MatIdentifier, &bAlreadyIn).AsInteger();
+
+		// identifier -> name association
+		if (!bAlreadyIn)
+		{
+			FName MaterialSlotName = *FString::FromInt(MatIdentifier);
+			GroupNamePerGroupIndex.Add(IndexOfIdentifier, MaterialSlotName);
+		}
+
+		// remap old identifier to material index
+		MatIdentifier = IndexOfIdentifier;
+	}
+
+	FMeshDescription MeshDescription;
+	FStaticMeshAttributes(MeshDescription).Register();
+
+	// Do not compute normals and tangents during conversion since we have other operations to apply
+	// on the mesh that might invalidate them anyway and we must also validate the mesh to detect
+	// vertex positions containing NaN before doing computation as MikkTSpace crashes on NaN values.
+	const bool bSkipNormalsAndTangents = true;
+	FStaticMeshOperations::ConvertFromRawMesh(RawMesh, MeshDescription, GroupNamePerGroupIndex, bSkipNormalsAndTangents);
+	return MeshDescription;
+}
+
+TArray<FDatasmithMeshModels> GetDatasmithMeshFromMeshPath_Legacy(FArchive* Archive, int32 LeagacyNumMeshesCount)
+{
+	TArray< FDatasmithMeshModels > Result;
+
+	UDatasmithMesh* DatasmithMesh = nullptr;
+	{
+		// Make sure the new UDatasmithMesh object is not created while a garbage collection is performed
+		FGCScopeGuard GCGuard;
+		// Setting the RF_Standalone bitmask on the new UDatasmithMesh object, to make sure it is not garbage collected
+		// while loading and processing the udsmesh file. This can happen on very big meshes (5M+ triangles)
+		DatasmithMesh = NewObject< UDatasmithMesh >(GetTransientPackage(), NAME_None, RF_Standalone);
+	}
+
+	// Currently we only have 1 mesh per file. If there's a second mesh, it will be a CollisionMesh
+	while ( LeagacyNumMeshesCount-- > 0)
+	{
+		TArray< uint8 > Bytes;
+		*Archive << Bytes;
+
+		FMemoryReader MemoryReader( Bytes, true );
+		MemoryReader.ArIgnoreClassRef = false;
+		MemoryReader.ArIgnoreArchetypeRef = false;
+		MemoryReader.SetWantBinaryPropertySerialization(true);
+		DatasmithMesh->Serialize( MemoryReader );
+
+		FDatasmithMeshModels& MeshInternal = Result.AddDefaulted_GetRef();
+		MeshInternal.bIsCollisionMesh = DatasmithMesh->bIsCollisionMesh;
+		for (FDatasmithMeshSourceModel& SourceModel : DatasmithMesh->SourceModels)
+		{
+			if (TOptional<FMeshDescription> OptionalMesh = ExtractToMeshDescription(SourceModel))
+			{
+				MeshInternal.SourceModels.Add(MoveTemp(*OptionalMesh));
+			}
+		}
+	}
+
+	// Tell the garbage collector DatasmithMesh can now be deleted.
+	DatasmithMesh->ClearInternalFlags(EInternalObjectFlags::Async);
+	DatasmithMesh->ClearFlags(RF_Standalone);
+	return Result;
+}
+
+TArray<FDatasmithMeshModels> GetDatasmithMeshFromMeshPath(const FString& MeshPath)
+{
+	TArray< FDatasmithMeshModels > Result;
+
+	TUniquePtr<FArchive> Archive( IFileManager::Get().CreateFileReader(*MeshPath) );
+	if ( !Archive.IsValid() )
+	{
+		UE_LOG(LogDatasmith, Warning, TEXT("Cannot read file %s"), *MeshPath);
+		return Result;
+	}
+
+	int32 LeagacyNumMeshesCount = 0;
+	*Archive << LeagacyNumMeshesCount;
+
+	if (LeagacyNumMeshesCount > 0)
+	{
+		return GetDatasmithMeshFromMeshPath_Legacy(Archive.Get(), LeagacyNumMeshesCount);
+	}
+
+	if (LeagacyNumMeshesCount == 0)
+	{
+		FDatasmithPackedMeshes Pack;
+		*Archive << Pack;
+
+		if (!Archive->IsError())
+		{
+			for (FDatasmithMeshModels& Mesh : Pack.MeshesToExport)
+			{
+				FDatasmithMeshModels& MeshInternal = Result.AddDefaulted_GetRef();
+				MeshInternal.bIsCollisionMesh = Mesh.bIsCollisionMesh;
+				MeshInternal.SourceModels = MoveTemp(Mesh.SourceModels);
+			}
+		}
+		else
+		{
+			UE_LOG(LogDatasmith, Warning, TEXT("Failed to read meshes from %s"), *MeshPath);
+		}
+	}
+
+	return Result;
+}
\ No newline at end of file
diff --git a/Engine/Source/Runtime/Datasmith/DatasmithCore/Private/DatasmithMeshUObject.cpp b/Engine/Source/Runtime/Datasmith/DatasmithCore/Private/DatasmithMeshUObject.cpp
--- a/Engine/Source/Runtime/Datasmith/DatasmithCore/Private/DatasmithMeshUObject.cpp
+++ b/Engine/Source/Runtime/Datasmith/DatasmithCore/Private/DatasmithMeshUObject.cpp
@@ -1,8 +1,6 @@
 // Copyright Epic Games, Inc. All Rights Reserved.
 
 #include "DatasmithMeshUObject.h"
-#include "Serialization/MemoryReader.h"
-#include "Serialization/MemoryWriter.h"
 
 void FDatasmithMeshSourceModel::SerializeBulkData(FArchive& Ar, UObject* Owner)
 {
@@ -18,53 +16,3 @@ void UDatasmithMesh::Serialize(FArchive& Ar)
 		SourceModel.SerializeBulkData( Ar, this );
 	}
 }
-
-void operator<<(FArchive& Ar, FDatasmithMeshModels& Models)
-{
-	Ar << Models.MeshName;
-	Ar << Models.bIsCollisionMesh;
-	Ar << Models.SourceModels;
-}
-
-void operator<<(FArchive& Ar, FDatasmithPackedMeshes& Pack)
-{
-	FString Guard = Ar.IsLoading() ? TEXT("") : TEXT("FDatasmithPackedMeshes");
-	Ar << Guard;
-	if (!ensure(Guard == TEXT("FDatasmithPackedMeshes")))
-	{
-		Ar.SetError();
-		return;
-	}
-
-	uint8 SerialVersionMajor = 1; // increment this when forward compatibility is not possible
-	Ar << SerialVersionMajor;
-	if (SerialVersionMajor > 1)
-	{
-		Ar.SetError();
-		return;
-	}
-
-	uint8 SerialVersionMinor = 0; // increment this to handle logic changes that preserves forward compatibility
-	Ar << SerialVersionMinor;
-
-	uint8 BufferType = 0;
-	Ar << BufferType; // (MeshDesc, Zipped Mesh desc;...)
-
-	if (Ar.IsLoading())
-	{
-		TArray<uint8> Bytes;
-		Ar << Bytes;
-		FMemoryReader Buffer(Bytes, true);
-		Buffer << Pack.MeshesToExport;
-	}
-	else
-	{
-		TArray<uint8> Bytes;
-		FMemoryWriter Buffer(Bytes, true);
-		Buffer << Pack.MeshesToExport;
-		Ar << Bytes;
-		FMD5 Md5;
-		Md5.Update(Bytes.GetData(), Bytes.Num());
-		Pack.OutHash.Set(Md5);
-	}
-}
diff --git a/Engine/Source/Runtime/Datasmith/DatasmithCore/Private/DatasmithUtils.cpp b/Engine/Source/Runtime/Datasmith/DatasmithCore/Private/DatasmithUtils.cpp
--- a/Engine/Source/Runtime/Datasmith/DatasmithCore/Private/DatasmithUtils.cpp
+++ b/Engine/Source/Runtime/Datasmith/DatasmithCore/Private/DatasmithUtils.cpp
@@ -20,6 +20,7 @@
 #include "StaticMeshAttributes.h"
 #include "StaticMeshOperations.h"
 #include "UObject/NameTypes.h"
+#include "UVMapSettings.h"
 
 static const float DATASMITH_FORMAT_VERSION = 0.24f; // Major.Minor - A change in the major version means that backward compatibility is broken
 
@@ -424,6 +425,87 @@ bool FDatasmithMeshUtils::IsUVChannelValid(const FDatasmithMesh& DsMesh, const i
 	return true;
 }
 
+void CreateDefaultUVs( FDatasmithMesh& Mesh )
+{
+	if ( Mesh.GetUVChannelsCount() > 0 )
+	{
+		return;
+	}
+
+	// Get the mesh description to generate BoxUV.
+	FMeshDescription MeshDescription;
+	FStaticMeshAttributes(MeshDescription).Register();
+	FDatasmithMeshUtils::ToMeshDescription(Mesh, MeshDescription);
+	FUVMapParameters UVParameters(
+		FVector(Mesh.GetExtents().GetCenter()),
+		FQuat::Identity,
+		FVector(Mesh.GetExtents().GetSize()),
+		FVector::OneVector,
+		FVector2D::UnitVector
+	);
+	TMap<FVertexInstanceID, FVector2D> TexCoords;
+	FStaticMeshOperations::GenerateBoxUV(MeshDescription, UVParameters, TexCoords);
+
+	// Put the results in a map to determine the number of unique values.
+	TMap<FVector2D, TArray<int32>> UniqueTexCoordMap;
+	for (const TPair<FVertexInstanceID, FVector2D>& Pair : TexCoords)
+	{
+		TArray<int32>& MappedIndices = UniqueTexCoordMap.FindOrAdd(Pair.Value);
+		MappedIndices.Add(Pair.Key.GetValue());
+	}
+
+	//Set the UV values
+	Mesh.AddUVChannel();
+	Mesh.SetUVCount(0, UniqueTexCoordMap.Num());
+	int32 UVIndex = 0;
+	TArray<int32> IndicesMapping;
+	IndicesMapping.AddZeroed(TexCoords.Num());
+	for (const TPair<FVector2D, TArray<int32>>& UniqueCoordPair : UniqueTexCoordMap)
+	{
+		Mesh.SetUV(0, UVIndex, UniqueCoordPair.Key.X, UniqueCoordPair.Key.Y);
+		for (int32 IndicesIndex : UniqueCoordPair.Value)
+		{
+			IndicesMapping[IndicesIndex] = UVIndex;
+		}
+		UVIndex++;
+	}
+
+	//Map the UV indices.
+	for (int32 FaceIndex = 0; FaceIndex < Mesh.GetFacesCount(); ++FaceIndex)
+	{
+		const int32 IndicesOffset = FaceIndex * 3;
+		check(IndicesOffset + 2 < IndicesMapping.Num());
+		Mesh.SetFaceUV(FaceIndex, 0, IndicesMapping[IndicesOffset + 0], IndicesMapping[IndicesOffset + 1], IndicesMapping[IndicesOffset + 2]);
+	}
+}
+
+void FDatasmithMeshUtils::CreateDefaultUVsWithLOD(FDatasmithMesh& Mesh)
+{
+	CreateDefaultUVs( Mesh );
+
+	for ( int32 LODIndex = 0; LODIndex < Mesh.GetLODsCount(); ++LODIndex )
+	{
+		if ( FDatasmithMesh* LODMesh = Mesh.GetLOD( LODIndex ) )
+		{
+			CreateDefaultUVs( *LODMesh );
+		}
+	}
+}
+
+void FDatasmithMeshUtils::ExtractVertexPositions(const FMeshDescription& Mesh, TArray<FVector3f>& OutPositions)
+{
+	FStaticMeshConstAttributes Attributes(Mesh);
+	const TVertexAttributesConstRef<FVector3f> VertexPositions = Attributes.GetVertexPositions();
+	if (VertexPositions.IsValid())
+	{
+		OutPositions.Reserve(VertexPositions.GetNumElements());
+		for (const FVertexID VertexID : Mesh.Vertices().GetElementIDs())
+		{
+			OutPositions.Add(VertexPositions[VertexID]);
+		}
+	}
+}
+
 bool FDatasmithTextureUtils::CalculateTextureHash(const TSharedPtr<class IDatasmithTextureElement>& TextureElement)
 {
 	if (TextureElement)
diff --git a/Engine/Source/Runtime/Datasmith/DatasmithCore/Public/DatasmithMeshSerialization.h b/Engine/Source/Runtime/Datasmith/DatasmithCore/Public/DatasmithMeshSerialization.h
new file mode 100644
--- /dev/null
+++ b/Engine/Source/Runtime/Datasmith/DatasmithCore/Public/DatasmithMeshSerialization.h
@@ -0,0 +1,26 @@
+// Copyright Epic Games, Inc. All Rights Reserved.
+
+#pragma once
+
+#include "DatasmithCore.h"
+#include "MeshDescription.h"
+#include "Misc/SecureHash.h"
+
+
+struct FDatasmithMeshModels
+{
+	FString MeshName;
+	bool bIsCollisionMesh;
+	TArray<FMeshDescription> SourceModels;
+
+	DATASMITHCORE_API friend void operator << (FArchive& Ar, FDatasmithMeshModels& Models);
+};
+
+struct FDatasmithPackedMeshes
+{
+	TArray<FDatasmithMeshModels> MeshesToExport;
+	DATASMITHCORE_API friend void operator << (FArchive& Ar, FDatasmithPackedMeshes& Pack);
+	FMD5Hash OutHash;
+};
+
+DATASMITHCORE_API TArray<FDatasmithMeshModels> GetDatasmithMeshFromMeshPath(const FString& MeshPath);
diff --git a/Engine/Source/Runtime/Datasmith/DatasmithCore/Public/DatasmithMeshUObject.h b/Engine/Source/Runtime/Datasmith/DatasmithCore/Public/DatasmithMeshUObject.h
--- a/Engine/Source/Runtime/Datasmith/DatasmithCore/Public/DatasmithMeshUObject.h
+++ b/Engine/Source/Runtime/Datasmith/DatasmithCore/Public/DatasmithMeshUObject.h
@@ -43,19 +43,3 @@ public:
 	UPROPERTY()
 	TArray< FDatasmithMeshSourceModel > SourceModels;
 };
-
-struct FDatasmithMeshModels
-{
-	FString MeshName;
-	bool bIsCollisionMesh;
-	TArray<FMeshDescription> SourceModels;
-
-	DATASMITHCORE_API friend void operator << (FArchive& Ar, FDatasmithMeshModels& Models);
-};
-
-struct FDatasmithPackedMeshes
-{
-	TArray<FDatasmithMeshModels> MeshesToExport;
-	DATASMITHCORE_API friend void operator << (FArchive& Ar, FDatasmithPackedMeshes& Pack);
-	FMD5Hash OutHash;
-};
\ No newline at end of file
diff --git a/Engine/Source/Runtime/Datasmith/DatasmithCore/Public/DatasmithUtils.h b/Engine/Source/Runtime/Datasmith/DatasmithCore/Public/DatasmithUtils.h
--- a/Engine/Source/Runtime/Datasmith/DatasmithCore/Public/DatasmithUtils.h
+++ b/Engine/Source/Runtime/Datasmith/DatasmithCore/Public/DatasmithUtils.h
@@ -145,6 +145,25 @@ public:
 	 * @param Channel	The UV channel to validate, starting at 0
 	 */
 	static bool IsUVChannelValid(const FDatasmithMesh& DsMesh, const int32 Channel);
+
+	/**
+	 * Generate simple UV data at channel 0.
+	 *
+	 * @param DatasmithMesh    The DatasmithMesh in which the UV data will be created.
+	 */
+	 // static void CreateDefaultUVs(FDatasmithMesh& Mesh);
+
+	/**
+	 * Generate simple UV data at channel 0 for the base mesh and it's various LOD variants.
+	 *
+	 * @param DatasmithMesh    The DatasmithMesh in which the UV data will be created.
+	 */
+	static void CreateDefaultUVsWithLOD(FDatasmithMesh& Mesh);
+
+	/**
+	 * Build an array of point from a MeshDescription
+	 */
+	static void ExtractVertexPositions(const FMeshDescription& Mesh, TArray<FVector3f>& OutPositions);
 };
 
 enum class EDSTextureUtilsError : int32
