DatasmithExporter - Fixing Rhino and Revit broken export (no export in Rhino, bad mesh scaling in Rhino & Revit)

- Restoring previous FDatasmithSceneExporter::Reset() behavior but marking the function as deprecated.
It's previous usage of reseting the export progress without resetting the whole exporter state was discutable and confusing. If users want to fully reset the exporter they don't need the Reset() function.

- Fixing Datasmith Facade inverse scaling error during position conversion

diff --git a/Engine/Source/Developer/Datasmith/DatasmithExporter/Private/DatasmithSceneExporter.cpp b/Engine/Source/Developer/Datasmith/DatasmithExporter/Private/DatasmithSceneExporter.cpp
--- a/Engine/Source/Developer/Datasmith/DatasmithExporter/Private/DatasmithSceneExporter.cpp
+++ b/Engine/Source/Developer/Datasmith/DatasmithExporter/Private/DatasmithSceneExporter.cpp
@@ -314,7 +314,10 @@ void FDatasmithSceneExporter::Export( TSharedRef< IDatasmithScene > DatasmithSce
 
 void FDatasmithSceneExporter::Reset()
 {
-	Impl = MakeUnique<FDatasmithSceneExporterImpl>();
+	Impl->ProgressManager = nullptr;
+	Impl->Logger = nullptr;
+
+	Impl->ExportStartCycles = 0;
 }
 
 void FDatasmithSceneExporter::SetProgressManager( const TSharedPtr< IDatasmithProgressManager >& InProgressManager )
diff --git a/Engine/Source/Developer/Datasmith/DatasmithExporter/Public/DatasmithSceneExporter.h b/Engine/Source/Developer/Datasmith/DatasmithExporter/Public/DatasmithSceneExporter.h
--- a/Engine/Source/Developer/Datasmith/DatasmithExporter/Public/DatasmithSceneExporter.h
+++ b/Engine/Source/Developer/Datasmith/DatasmithExporter/Public/DatasmithSceneExporter.h
@@ -37,6 +37,7 @@ public:
 	void Export( TSharedRef< IDatasmithScene > DatasmithScene, bool bCleanupUnusedElements = true );
 
 	/** Resets all the settings on the scene */
+	UE_DEPRECATED(5.1, "This function was selectively reseting the export state and was preserving the export paths. We should now directly set the fields we want to change instead.")
 	void Reset();
 
 	/** Sets the progress manager for visual feedback on exporting */
diff --git a/Engine/Source/Developer/Datasmith/DatasmithFacade/Private/DatasmithFacadeElement.cpp b/Engine/Source/Developer/Datasmith/DatasmithFacade/Private/DatasmithFacadeElement.cpp
--- a/Engine/Source/Developer/Datasmith/DatasmithFacade/Private/DatasmithFacadeElement.cpp
+++ b/Engine/Source/Developer/Datasmith/DatasmithFacade/Private/DatasmithFacadeElement.cpp
@@ -81,7 +81,7 @@ Vec_t FDatasmithFacadeElement::Convert(const Vec_t& V)
 	// World scaling for positions conversions
 	if constexpr (!kIsDirection)
 	{
-		Tmp *= kIsForward ? 1 / WorldUnitScale : WorldUnitScale;
+		Tmp *= kIsForward ? WorldUnitScale : 1 / WorldUnitScale;
 	}
 
 	return Tmp;
diff --git a/Engine/Source/Developer/Datasmith/DatasmithFacade/Private/DatasmithFacadeScene.cpp b/Engine/Source/Developer/Datasmith/DatasmithFacade/Private/DatasmithFacadeScene.cpp
--- a/Engine/Source/Developer/Datasmith/DatasmithFacade/Private/DatasmithFacadeScene.cpp
+++ b/Engine/Source/Developer/Datasmith/DatasmithFacade/Private/DatasmithFacadeScene.cpp
@@ -452,9 +452,6 @@ void FDatasmithFacadeScene::PreExport()
 	// Initialize the Datasmith exporter module.
 	FDatasmithExporterManager::Initialize();
 
-	// Create a Datasmith scene exporter.
-	SceneExporterRef->Reset();
-
 	// Start measuring the time taken to export the scene.
 	SceneExporterRef->PreExport();
 }
diff --git a/Engine/Source/Programs/Enterprise/Datasmith/DatasmithMaxExporter/Private/DatasmithMaxExporterWithDirectLink.cpp b/Engine/Source/Programs/Enterprise/Datasmith/DatasmithMaxExporter/Private/DatasmithMaxExporterWithDirectLink.cpp
--- a/Engine/Source/Programs/Enterprise/Datasmith/DatasmithMaxExporter/Private/DatasmithMaxExporterWithDirectLink.cpp
+++ b/Engine/Source/Programs/Enterprise/Datasmith/DatasmithMaxExporter/Private/DatasmithMaxExporterWithDirectLink.cpp
@@ -126,9 +126,6 @@ public:
 
 	void PreExport()
 	{
-		// Create a Datasmith scene exporter.
-		SceneExporterRef->Reset();
-
 		// Start measuring the time taken to export the scene.
 		SceneExporterRef->PreExport();
 	}
diff --git a/Engine/Source/Programs/Enterprise/Datasmith/DatasmithSketchUpRubyExporter/Private/DatasmithSketchUpExporter.cpp b/Engine/Source/Programs/Enterprise/Datasmith/DatasmithSketchUpRubyExporter/Private/DatasmithSketchUpExporter.cpp
--- a/Engine/Source/Programs/Enterprise/Datasmith/DatasmithSketchUpRubyExporter/Private/DatasmithSketchUpExporter.cpp
+++ b/Engine/Source/Programs/Enterprise/Datasmith/DatasmithSketchUpRubyExporter/Private/DatasmithSketchUpExporter.cpp
@@ -148,9 +148,6 @@ public:
 
 	void PreExport()
 	{
-		// Create a Datasmith scene exporter.
-		SceneExporterRef->Reset();
-
 		// Start measuring the time taken to export the scene.
 		SceneExporterRef->PreExport();
 	}
