DatasmithSDK: Default to new mesh serialization
- requires 5.0.3+ to load

diff --git a/Engine/Source/Developer/Datasmith/DatasmithExporter/Private/DatasmithMeshExporter.cpp b/Engine/Source/Developer/Datasmith/DatasmithExporter/Private/DatasmithMeshExporter.cpp
--- a/Engine/Source/Developer/Datasmith/DatasmithExporter/Private/DatasmithMeshExporter.cpp
+++ b/Engine/Source/Developer/Datasmith/DatasmithExporter/Private/DatasmithMeshExporter.cpp
@@ -297,6 +297,7 @@ TSharedPtr< IDatasmithMeshElement > FDatasmithMeshExporter::ExportToUObject(cons
 	}
 	else
 	{
+		ensureMsgf(false, TEXT("Old serialization used for mesh %s. It is no longer supported."), *FullPath);
 		LegacyImpl->DoExport( ExportedMeshElement, ExportOptions );
 	}
 
@@ -315,6 +316,7 @@ bool FDatasmithMeshExporter::ExportToUObject( TSharedPtr< IDatasmithMeshElement
 	}
 	else
 	{
+		ensureMsgf(false, TEXT("Old serialization used for mesh %s. It is no longer supported."), MeshElement->GetName());
 		return LegacyImpl->DoExport( MeshElement, ExportOptions );
 	}
 }
diff --git a/Engine/Source/Developer/Datasmith/DatasmithExporter/Public/DatasmithMeshExporter.h b/Engine/Source/Developer/Datasmith/DatasmithExporter/Public/DatasmithMeshExporter.h
--- a/Engine/Source/Developer/Datasmith/DatasmithExporter/Public/DatasmithMeshExporter.h
+++ b/Engine/Source/Developer/Datasmith/DatasmithExporter/Public/DatasmithMeshExporter.h
@@ -24,7 +24,7 @@ public:
 	 *
 	 * @return				A IDatasmithMeshElement that refers to the exported file
 	 */
-	TSharedPtr< IDatasmithMeshElement > ExportToUObject( const TCHAR* Filepath, const TCHAR* Filename, FDatasmithMesh& Mesh, FDatasmithMesh* CollisionMesh, EDSExportLightmapUV LightmapUV, bool NewFormat=false );
+	TSharedPtr< IDatasmithMeshElement > ExportToUObject( const TCHAR* Filepath, const TCHAR* Filename, FDatasmithMesh& Mesh, FDatasmithMesh* CollisionMesh, EDSExportLightmapUV LightmapUV, bool NewFormat=true );
 
 	/**
 	 * Exports a FDatasmithMesh as a UObject and link it to the given IDatasmithMeshElementElement.
@@ -37,7 +37,7 @@ public:
 	 *
 	 * @return				True if export was successful.
 	 */
-	bool ExportToUObject( TSharedPtr< IDatasmithMeshElement >& MeshElement, const TCHAR* Filepath, FDatasmithMesh& Mesh, FDatasmithMesh* CollisionMesh, EDSExportLightmapUV LightmapUV, bool NewFormat=false );
+	bool ExportToUObject( TSharedPtr< IDatasmithMeshElement >& MeshElement, const TCHAR* Filepath, FDatasmithMesh& Mesh, FDatasmithMesh* CollisionMesh, EDSExportLightmapUV LightmapUV, bool NewFormat=true );
 
 	/**
 	 * @return The error that happened during the last export, if any
