Fixed missing instances in ArchiCAD exporter
 - Root cause: Multi-threading, solution: temporarily disable multi-htreading.
Added UE natvis to ArchiCAD solution for debug purpose
Removed unnecessary library
Cleaned up syntax from tread to thread
Added logic to setting of Direct Link cache directories: No check was made on their validity
Simplified build steps to remove useless compilation of DatasmithArchiCAD project

diff --git a/Engine/Source/Developer/Datasmith/DatasmithExporterUI/Private/DirectLinkUI.cpp b/Engine/Source/Developer/Datasmith/DatasmithExporterUI/Private/DirectLinkUI.cpp
--- a/Engine/Source/Developer/Datasmith/DatasmithExporterUI/Private/DirectLinkUI.cpp
+++ b/Engine/Source/Developer/Datasmith/DatasmithExporterUI/Private/DirectLinkUI.cpp
@@ -12,6 +12,7 @@
 #include "Delegates/Delegate.h"
 #include "Framework/Application/SWindowTitleBar.h"
 #include "Framework/Application/SlateApplication.h"
+#include "HAL/FileManager.h"
 #include "HAL/Platform.h"
 #include "Misc/ConfigCacheIni.h"
 #include "Misc/ScopeLock.h"
@@ -60,21 +61,43 @@ namespace DirectLinkUIUtils
 
 FDirectLinkUI::FDirectLinkUI()
 {
-	FString ConfigPath = DirectLinkUIUtils::GetConfigPath();
-	FString DirectLinkCacheSectionAndValue = DirectLinkUIUtils::GetConfigCacheDirectorySectionAndValue( false );
-	FString DirectLinkDefaultCacheSectionAndValue = DirectLinkUIUtils::GetConfigCacheDirectorySectionAndValue( true );
+	const FString ConfigPath = DirectLinkUIUtils::GetConfigPath();
+	const FString DirectLinkCacheSectionAndValue = DirectLinkUIUtils::GetConfigCacheDirectorySectionAndValue( false );
+	const FString DirectLinkDefaultCacheSectionAndValue = DirectLinkUIUtils::GetConfigCacheDirectorySectionAndValue( true );
 
 	FScopeLock Lock( &CriticalSectionCacheDirectory );
-	if ( !GConfig->GetString( *DirectLinkDefaultCacheSectionAndValue, *DirectLinkDefaultCacheSectionAndValue, DefaultDirectLinkCacheDirectory, ConfigPath ) )
+
+	// Verify and get default cache directory
+	GConfig->GetString(*DirectLinkDefaultCacheSectionAndValue, *DirectLinkDefaultCacheSectionAndValue, DefaultDirectLinkCacheDirectory, ConfigPath);
+
+	if (DefaultDirectLinkCacheDirectory.IsEmpty() || !FPaths::DirectoryExists(DefaultDirectLinkCacheDirectory))
 	{
-		DefaultDirectLinkCacheDirectory = FPaths::Combine( FPlatformProcess::UserTempDir(), TEXT("DLExporter") );
-		SaveCacheDirectory( DefaultDirectLinkCacheDirectory, true );
+		if (DefaultDirectLinkCacheDirectory.IsEmpty() || !IFileManager::Get().MakeDirectory(*DefaultDirectLinkCacheDirectory, true))
+		{
+			DefaultDirectLinkCacheDirectory = FPaths::Combine(FPlatformProcess::UserTempDir(), TEXT("DLExporter"));
+			if (!FPaths::DirectoryExists(DefaultDirectLinkCacheDirectory))
+			{
+				ensure(IFileManager::Get().MakeDirectory(*DefaultDirectLinkCacheDirectory, true));
+			}
+		}
+
 	}
-	if ( !GConfig->GetString( *DirectLinkCacheSectionAndValue, *DirectLinkCacheSectionAndValue, DirectLinkCacheDirectory, ConfigPath ) )
+
+	GConfig->SetString(*DirectLinkDefaultCacheSectionAndValue, *DirectLinkDefaultCacheSectionAndValue, *DefaultDirectLinkCacheDirectory, ConfigPath);
+
+	// Verify and get current cache directory
+	GConfig->GetString(*DirectLinkCacheSectionAndValue, *DirectLinkCacheSectionAndValue, DirectLinkCacheDirectory, ConfigPath);
+
+	if (DirectLinkCacheDirectory.IsEmpty() || !FPaths::DirectoryExists(DirectLinkCacheDirectory))
 	{
-		DirectLinkCacheDirectory = DefaultDirectLinkCacheDirectory;
-		SaveCacheDirectory( DirectLinkCacheDirectory, false );
+		if (DirectLinkCacheDirectory.IsEmpty() || !IFileManager::Get().MakeDirectory(*DirectLinkCacheDirectory, true))
+		{
+			DirectLinkCacheDirectory = DefaultDirectLinkCacheDirectory;
+		}
+
 	}
+
+	GConfig->SetString(*DirectLinkCacheSectionAndValue, *DirectLinkCacheSectionAndValue, *DirectLinkCacheDirectory, ConfigPath);
 }
 
 void FDirectLinkUI::OpenDirectLinkStreamWindow()
diff --git a/Engine/Source/Programs/Enterprise/Datasmith/DatasmithARCHICADExporter/Private/Utils/TaskMgr.cpp b/Engine/Source/Programs/Enterprise/Datasmith/DatasmithARCHICADExporter/Private/Utils/TaskMgr.cpp
--- a/Engine/Source/Programs/Enterprise/Datasmith/DatasmithARCHICADExporter/Private/Utils/TaskMgr.cpp
+++ b/Engine/Source/Programs/Enterprise/Datasmith/DatasmithARCHICADExporter/Private/Utils/TaskMgr.cpp
@@ -104,22 +104,24 @@ class FTaskMgr::FThread
 // Constructor
 FTaskMgr::FTaskMgr()
 	: bTerminate(false)
-	, bTreadingEnabled(true)
+	// Temporarily disabling multi-threading to fix UE-127453 for 5.0
+	// #ue_archicad : Must revisit multi-threading to find proper solution for UE-127453 
+	, bThreadingEnabled(false)
 	, NbTaskWaiting(0)
 	, CV(AccessControl)
 {
-	if (bTreadingEnabled)
+	if (bThreadingEnabled)
 	{
 		// One thread by processor
 		unsigned nbProcessors = GS::System::GetNumberOfActiveProcessors();
 		if (nbProcessors == 0)
 		{
-			nbProcessors = 1;
+			bThreadingEnabled = false;
 		}
-		Treads.resize(nbProcessors);
+		Threads.resize(nbProcessors);
 		for (unsigned i = 0; i < nbProcessors; i++)
 		{
-			Treads[i].reset(new FThread(this, GS::UniString::Printf("Datasmith Exporter Task #%d", i)));
+			Threads[i].reset(new FThread(this, GS::UniString::Printf("Datasmith Exporter Task #%d", i)));
 		}
 	}
 }
@@ -133,10 +135,10 @@ FTaskMgr::~FTaskMgr()
 	}
 	CV.NotifyAll();
 
-	size_t nbThreads = Treads.size();
+	size_t nbThreads = Threads.size();
 	for (size_t i = 0; i < nbThreads; i++)
 	{
-		Treads[i]->Join();
+		Threads[i]->Join();
 	}
 
 	// Delete all not processed task. This happen when Join has'n been called on the thread manager
@@ -151,7 +153,7 @@ FTaskMgr::~FTaskMgr()
 // Add task
 void FTaskMgr::AddTask(FTask* InTask, ERunMode inMode)
 {
-	if (bTreadingEnabled && inMode == kSchedule)
+	if (bThreadingEnabled && inMode == kSchedule)
 	{
 		GS::Guard< GS::Lock > lck(AccessControl);
 		if (bTerminate)
@@ -172,7 +174,7 @@ void FTaskMgr::AddTask(FTask* InTask, ERunMode inMode)
 void FTaskMgr::Join(FProgression* inProgression)
 {
 	GS::Guard< GS::Lock > lck(AccessControl);
-	while (!TaskQueue.empty() || NbTaskWaiting != Treads.size())
+	while (!TaskQueue.empty() || NbTaskWaiting != Threads.size())
 	{
 		CV.NotifyAll();
 		CV.Wait(10, true, nullptr);
diff --git a/Engine/Source/Programs/Enterprise/Datasmith/DatasmithARCHICADExporter/Private/Utils/TaskMgr.h b/Engine/Source/Programs/Enterprise/Datasmith/DatasmithARCHICADExporter/Private/Utils/TaskMgr.h
--- a/Engine/Source/Programs/Enterprise/Datasmith/DatasmithARCHICADExporter/Private/Utils/TaskMgr.h
+++ b/Engine/Source/Programs/Enterprise/Datasmith/DatasmithARCHICADExporter/Private/Utils/TaskMgr.h
@@ -31,14 +31,14 @@ class FTaskMgr
 	// Destructor
 	~FTaskMgr();
 
-	// Task can be added for schedule or run immediatly in the current thread).
+	// Task can be added for schedule or run immediately in the current thread).
 	enum ERunMode
 	{
 		kSchedule,
 		kRunNow
 	};
 
-	// Add task (if inRunNow, the task is run immediatly in current thread)
+	// Add task (if inRunNow, the task is run immediately in current thread)
 	void AddTask(FTask* InTask, ERunMode InMode = kSchedule);
 
 	// Wait until all task have been processed
@@ -46,8 +46,8 @@ class FTaskMgr
 
 	bool EnableThreading(bool InEnable)
 	{
-		bool tmp = bTreadingEnabled;
-		bTreadingEnabled = InEnable;
+		bool tmp = bThreadingEnabled;
+		bThreadingEnabled = InEnable;
 		return tmp;
 	}
 
@@ -62,8 +62,8 @@ class FTaskMgr
 	// Class that implement CTaskMgr thread
 	class FThread;
 
-	// Threads used (≈ nomber of cpu)
-	std::vector< std::unique_ptr< FThread > > Treads;
+	// Threads used (≈ number of cpu)
+	std::vector< std::unique_ptr< FThread > > Threads;
 
 	// Fifo tasks queue
 	std::queue< FTask* > TaskQueue;
@@ -81,7 +81,7 @@ class FTaskMgr
 	volatile int NbTaskWaiting;
 
 	// False: tasks to run in main threads, true: Task run in thread
-	bool bTreadingEnabled;
+	bool bThreadingEnabled;
 };
 
 END_NAMESPACE_UE_AC
