Datasmith: Mesh Exporter: add a way to silent the deprecation ensure
#rnx

diff --git a/Engine/Source/Developer/Datasmith/DatasmithExporter/Private/DatasmithMeshExporter.cpp b/Engine/Source/Developer/Datasmith/DatasmithExporter/Private/DatasmithMeshExporter.cpp
--- a/Engine/Source/Developer/Datasmith/DatasmithExporter/Private/DatasmithMeshExporter.cpp
+++ b/Engine/Source/Developer/Datasmith/DatasmithExporter/Private/DatasmithMeshExporter.cpp
@@ -219,9 +219,10 @@ bool FDatasmithMeshExporterLegacyImpl::ExportMeshes( const FDatasmithMeshExporte
 }
 
 
-FDatasmithMeshExporter::FDatasmithMeshExporter()
+FDatasmithMeshExporter::FDatasmithMeshExporter(bool bAllowOldSerialization)
 	: LegacyImpl(MakeUnique<FDatasmithMeshExporterLegacyImpl>())
 	, Impl(MakeUnique<FDatasmithMeshExporterImpl>())
+	, bAllowOldSerialization(bAllowOldSerialization)
 {}
 
 FDatasmithMeshExporter::~FDatasmithMeshExporter() = default;
@@ -297,7 +298,7 @@ TSharedPtr< IDatasmithMeshElement > FDatasmithMeshExporter::ExportToUObject(cons
 	}
 	else
 	{
-		ensureMsgf(false, TEXT("Old serialization used for mesh %s. It is no longer supported."), *FullPath);
+		ensureMsgf(bAllowOldSerialization, TEXT("Old serialization used for mesh %s. It is no longer supported."), *FullPath);
 		LegacyImpl->DoExport( ExportedMeshElement, ExportOptions );
 	}
 
@@ -316,7 +317,7 @@ bool FDatasmithMeshExporter::ExportToUObject( TSharedPtr< IDatasmithMeshElement
 	}
 	else
 	{
-		ensureMsgf(false, TEXT("Old serialization used for mesh %s. It is no longer supported."), MeshElement->GetName());
+		ensureMsgf(bAllowOldSerialization, TEXT("Old serialization used for mesh %s. It is no longer supported."), MeshElement->GetName());
 		return LegacyImpl->DoExport( MeshElement, ExportOptions );
 	}
 }
diff --git a/Engine/Source/Developer/Datasmith/DatasmithExporter/Public/DatasmithMeshExporter.h b/Engine/Source/Developer/Datasmith/DatasmithExporter/Public/DatasmithMeshExporter.h
--- a/Engine/Source/Developer/Datasmith/DatasmithExporter/Public/DatasmithMeshExporter.h
+++ b/Engine/Source/Developer/Datasmith/DatasmithExporter/Public/DatasmithMeshExporter.h
@@ -10,7 +10,8 @@ class IDatasmithMeshElement;
 class DATASMITHEXPORTER_API FDatasmithMeshExporter
 {
 public:
-	FDatasmithMeshExporter();
+	FDatasmithMeshExporter(bool bAllowOldSerialization);
+	FDatasmithMeshExporter() : FDatasmithMeshExporter(false) {}
 	virtual ~FDatasmithMeshExporter();
 
 	/**
@@ -47,4 +48,5 @@ public:
 private:
 	TUniquePtr<class FDatasmithMeshExporterLegacyImpl> LegacyImpl;
 	TUniquePtr<class FDatasmithMeshExporterImpl> Impl;
+	bool bAllowOldSerialization = false;
 };
