DatasmithExporter - Fixing Datasmith GameThread initialization.

The Datasmith GameThread + DirectLink initialization was trigerring some ensure which where sometimes not handled properly and could cause a crash.
We were not properly tagging Unreal's main thread and the dcc's main thread, causing those thread to compete for the main thread role.

diff --git a/Engine/Source/Developer/Datasmith/DatasmithExporter/Private/DatasmithExporterManager.cpp b/Engine/Source/Developer/Datasmith/DatasmithExporter/Private/DatasmithExporterManager.cpp
--- a/Engine/Source/Developer/Datasmith/DatasmithExporter/Private/DatasmithExporterManager.cpp
+++ b/Engine/Source/Developer/Datasmith/DatasmithExporter/Private/DatasmithExporterManager.cpp
@@ -162,6 +162,9 @@ FDatasmithGameThread::FDatasmithGameThread(FString&& InPreInitCommandArgs, const
 
 uint32 FDatasmithGameThread::Run()
 {
+	// We need to explicitly flag this thread as the GameThread to avoid errors.
+	FTaskTagScope Scope(ETaskTag::EGameThread);
+
 	// init
 	bKeepRunning = true;
 	OnInit();
@@ -327,9 +330,13 @@ bool FDatasmithExporterManager::Initialize(const FInitOptions& InitOptions)
 			TFuture<bool> OnInitDoneFuture = GDatasmithGameThread->GetOnInitDoneFuture();
 			GMainThreadAsRunnable = FRunnableThread::Create(GDatasmithGameThread.Get(), TEXT("DatasmithMainThread"));
 			bEngineInitialized = OnInitDoneFuture.Get();
+
+			// Remove the ETaskTag::EStaticInit tag on the current thread as it can cause this thread to compete with the main thread we just created.
+			FTaskTagScope::SetTagNone();
 		}
 		else
 		{
+			FTaskTagScope Scope(ETaskTag::EGameThread);
 			bEngineInitialized = DatasmithGameThread::InitializeInCurrentThread(CmdLine, InitOptions.bSuppressLogs);
 		}
 
