Revit - Connection and Message dialogs positioning not remembered
Connection/Messages/Settings windows and dialogs now remember position and size within the same session

diff --git a/Engine/Source/Developer/Datasmith/DatasmithExporterUI/Private/DirectLinkUI.cpp b/Engine/Source/Developer/Datasmith/DatasmithExporterUI/Private/DirectLinkUI.cpp
--- a/Engine/Source/Developer/Datasmith/DatasmithExporterUI/Private/DirectLinkUI.cpp
+++ b/Engine/Source/Developer/Datasmith/DatasmithExporterUI/Private/DirectLinkUI.cpp
@@ -59,7 +59,10 @@ namespace DirectLinkUIUtils
 	}
 }
 
+FVector2D FDirectLinkUI::StreamWindowDefaultSize = FVector2D(640, 480);
+
 FDirectLinkUI::FDirectLinkUI()
+	: StreamWindowSize(StreamWindowDefaultSize)
 {
 	const FString ConfigPath = DirectLinkUIUtils::GetConfigPath();
 	const FString DirectLinkCacheSectionAndValue = DirectLinkUIUtils::GetConfigCacheDirectorySectionAndValue( false );
@@ -100,6 +103,12 @@ FDirectLinkUI::FDirectLinkUI()
 	GConfig->SetString(*DirectLinkCacheSectionAndValue, *DirectLinkCacheSectionAndValue, *DirectLinkCacheDirectory, ConfigPath);
 }
 
+void FDirectLinkUI::SetStreamWindowCenter( int InCenterX, int InCenterY )
+{
+	StreamWindowCenterSet = true;
+	StreamWindowPosition = FVector2D(InCenterX, InCenterY) - (StreamWindowSize / 2);
+}
+
 void FDirectLinkUI::OpenDirectLinkStreamWindow()
 {
 	FSimpleDelegate RunOnUIThread;
@@ -118,15 +127,22 @@ void FDirectLinkUI::OpenDirectLinkStreamWindow()
 		}
 		else
 		{
+			EAutoCenter AutoCenterValue = StreamWindowClosedBefore ? (EAutoCenter::None) : (StreamWindowCenterSet ? EAutoCenter::None : EAutoCenter::PrimaryWorkArea);
+
 			// This window setup might be an issue on mac where users often expect os borders on their window
 			TSharedRef<SWindow> Window = SNew( SWindow )
 				.CreateTitleBar( false )
-				.ClientSize( FVector2D( 640, 480 ) )
-				.AutoCenter( EAutoCenter::PrimaryWorkArea )
+				.ClientSize( StreamWindowDefaultSize )
+				.AutoCenter( AutoCenterValue )
 				.SizingRule( ESizingRule::UserSized )
 				.FocusWhenFirstShown( true )
 				.Title( LOCTEXT("DirectlinkStreamManagerWindowTitle", "Datasmith Direct Link Connection Status") );
 
+			if ( AutoCenterValue == EAutoCenter::None )
+			{
+				Window->ReshapeWindow( StreamWindowPosition, StreamWindowSize );
+			}
+			
 			TSharedRef<SWindowTitleBar> WindowTitleBar = SNew( SWindowTitleBar, Window, nullptr, DirectLinkUIUtils::GetWindowTitleAlignement() )
 				.Visibility( EVisibility::Visible )
 				.ShowAppIcon( false );
@@ -155,7 +171,7 @@ void FDirectLinkUI::OpenDirectLinkStreamWindow()
 						]
 					]
 				);
-
+			Window->SetOnWindowClosed( FOnWindowClosed::CreateRaw(this, &FDirectLinkUI::WindowClosed) );
 			DirectLinkWindow = Window;
 			FSlateApplication::Get().AddWindow( Window, true );
 			Window->HACK_ForceToFront();
@@ -169,6 +185,14 @@ void FDirectLinkUI::OpenDirectLinkStreamWindow()
 #endif
 }
 
+void FDirectLinkUI::WindowClosed( const TSharedRef<SWindow>& WindowArg )
+{
+	StreamWindowClosedBefore = true;
+
+	StreamWindowSize = WindowArg->GetSizeInScreen();
+	StreamWindowPosition = WindowArg->GetPositionInScreen();
+}
+
 const TCHAR* FDirectLinkUI::GetDirectLinkCacheDirectory()
 {
 	// This function should always be called from one thread only
diff --git a/Engine/Source/Developer/Datasmith/DatasmithExporterUI/Private/DirectLinkUI.h b/Engine/Source/Developer/Datasmith/DatasmithExporterUI/Private/DirectLinkUI.h
--- a/Engine/Source/Developer/Datasmith/DatasmithExporterUI/Private/DirectLinkUI.h
+++ b/Engine/Source/Developer/Datasmith/DatasmithExporterUI/Private/DirectLinkUI.h
@@ -5,6 +5,7 @@
 #include "IDirectLinkUI.h"
 
 #include "Containers/UnrealString.h"
+#include "Math/Vector2D.h"
 #include "HAL/CriticalSection.h"
 #include "Templates/SharedPointer.h"
 
@@ -16,6 +17,8 @@ public:
 
 	FDirectLinkUI();
 
+	virtual void SetStreamWindowCenter( int InCenterX, int InCenterY ) override;
+
 	virtual void OpenDirectLinkStreamWindow() override;
 	virtual const TCHAR* GetDirectLinkCacheDirectory() override;
 
@@ -25,6 +28,8 @@ private:
 
 	void SaveCacheDirectory( const FString& InCacheDir, bool bInDefaultCacheDir );
 
+	void WindowClosed( const TSharedRef<SWindow>& WindowArg );
+
 	// To accessed from the game thread only
 	TWeakPtr<SWindow> DirectLinkWindow;
 	
@@ -34,4 +39,10 @@ private:
 
 	// Used to protect the caller to GetDirectLinkCacheDirectory from a pontential race condition
 	FString LastReturnedCacheDirectory;
+
+	bool StreamWindowCenterSet = false;
+	bool StreamWindowClosedBefore = false;
+	static FVector2D StreamWindowDefaultSize;
+	FVector2D StreamWindowSize;
+	FVector2D StreamWindowPosition;
 };
\ No newline at end of file
diff --git a/Engine/Source/Developer/Datasmith/DatasmithExporterUI/Public/IDirectLinkUI.h b/Engine/Source/Developer/Datasmith/DatasmithExporterUI/Public/IDirectLinkUI.h
--- a/Engine/Source/Developer/Datasmith/DatasmithExporterUI/Public/IDirectLinkUI.h
+++ b/Engine/Source/Developer/Datasmith/DatasmithExporterUI/Public/IDirectLinkUI.h
@@ -10,6 +10,7 @@ public:
 	/**
 	 * Open the direct link connections window or move it up front if it's already open
 	 */
+	virtual void SetStreamWindowCenter(int InCenterX, int InCenterY) = 0;
 	virtual void OpenDirectLinkStreamWindow() = 0;
 
 	/**
diff --git a/Engine/Source/Programs/Enterprise/Datasmith/DatasmithFacadeCSharp/Private/DatasmithFacadeCSharp.cpp b/Engine/Source/Programs/Enterprise/Datasmith/DatasmithFacadeCSharp/Private/DatasmithFacadeCSharp.cpp
--- a/Engine/Source/Programs/Enterprise/Datasmith/DatasmithFacadeCSharp/Private/DatasmithFacadeCSharp.cpp
+++ b/Engine/Source/Programs/Enterprise/Datasmith/DatasmithFacadeCSharp/Private/DatasmithFacadeCSharp.cpp
@@ -7465,6 +7465,18 @@ SWIGEXPORT void SWIGSTDCALL CSharp_delete_IDatasmithExporterUIModule(void * jarg
 }
 
 
+SWIGEXPORT void SWIGSTDCALL CSharp_IDirectLinkUI_SetStreamWindowCenter(void * jarg1, int jarg2, int jarg3) {
+  IDirectLinkUI *arg1 = (IDirectLinkUI *) 0 ;
+  int arg2 ;
+  int arg3 ;
+  
+  arg1 = (IDirectLinkUI *)jarg1; 
+  arg2 = (int)jarg2; 
+  arg3 = (int)jarg3; 
+  (arg1)->SetStreamWindowCenter(arg2,arg3);
+}
+
+
 SWIGEXPORT void SWIGSTDCALL CSharp_IDirectLinkUI_OpenDirectLinkStreamWindow(void * jarg1) {
   IDirectLinkUI *arg1 = (IDirectLinkUI *) 0 ;
   
diff --git a/Engine/Source/Programs/Enterprise/Datasmith/DatasmithFacadeCSharp/Public/DatasmithFacadeCSharpPINVOKE.cs b/Engine/Source/Programs/Enterprise/Datasmith/DatasmithFacadeCSharp/Public/DatasmithFacadeCSharpPINVOKE.cs
--- a/Engine/Source/Programs/Enterprise/Datasmith/DatasmithFacadeCSharp/Public/DatasmithFacadeCSharpPINVOKE.cs
+++ b/Engine/Source/Programs/Enterprise/Datasmith/DatasmithFacadeCSharp/Public/DatasmithFacadeCSharpPINVOKE.cs
@@ -2060,6 +2060,9 @@ class DatasmithFacadeCSharpPINVOKE {
   [global::System.Runtime.InteropServices.DllImport("DatasmithFacadeCSharp", EntryPoint="CSharp_delete_IDatasmithExporterUIModule")]
   public static extern void delete_IDatasmithExporterUIModule(global::System.Runtime.InteropServices.HandleRef jarg1);
 
+  [global::System.Runtime.InteropServices.DllImport("DatasmithFacadeCSharp", EntryPoint="CSharp_IDirectLinkUI_SetStreamWindowCenter")]
+  public static extern void IDirectLinkUI_SetStreamWindowCenter(global::System.Runtime.InteropServices.HandleRef jarg1, int jarg2, int jarg3);
+
   [global::System.Runtime.InteropServices.DllImport("DatasmithFacadeCSharp", EntryPoint="CSharp_IDirectLinkUI_OpenDirectLinkStreamWindow")]
   public static extern void IDirectLinkUI_OpenDirectLinkStreamWindow(global::System.Runtime.InteropServices.HandleRef jarg1);
 
diff --git a/Engine/Source/Programs/Enterprise/Datasmith/DatasmithFacadeCSharp/Public/IDirectLinkUI.cs b/Engine/Source/Programs/Enterprise/Datasmith/DatasmithFacadeCSharp/Public/IDirectLinkUI.cs
--- a/Engine/Source/Programs/Enterprise/Datasmith/DatasmithFacadeCSharp/Public/IDirectLinkUI.cs
+++ b/Engine/Source/Programs/Enterprise/Datasmith/DatasmithFacadeCSharp/Public/IDirectLinkUI.cs
@@ -45,6 +45,10 @@ public class IDirectLinkUI : global::System.IDisposable {
     }
   }
 
+  public virtual void SetStreamWindowCenter(int InCenterX, int InCenterY) {
+    DatasmithFacadeCSharpPINVOKE.IDirectLinkUI_SetStreamWindowCenter(swigCPtr, InCenterX, InCenterY);
+  }
+
   public virtual void OpenDirectLinkStreamWindow() {
     DatasmithFacadeCSharpPINVOKE.IDirectLinkUI_OpenDirectLinkStreamWindow(swigCPtr);
   }
diff --git a/Engine/Source/Programs/Enterprise/Datasmith/DatasmithRevitExporter/Private/DatasmithRevitApplication.cs b/Engine/Source/Programs/Enterprise/Datasmith/DatasmithRevitExporter/Private/DatasmithRevitApplication.cs
--- a/Engine/Source/Programs/Enterprise/Datasmith/DatasmithRevitExporter/Private/DatasmithRevitApplication.cs
+++ b/Engine/Source/Programs/Enterprise/Datasmith/DatasmithRevitExporter/Private/DatasmithRevitApplication.cs
@@ -414,11 +414,14 @@ namespace DatasmithRevitExporter
 			}
 		}
 
-		public static void ShowExportMessages()
+		public static void ShowExportMessages(ExternalCommandData InCommandData)
 		{
 			if (ExportMessagesDialog == null || ExportMessagesDialog.IsDisposed)
 			{
-				ExportMessagesDialog = new DatasmithRevitExportMessages(() => ExportMessages = "");
+				int CenterX = (InCommandData.Application.MainWindowExtents.Left + InCommandData.Application.MainWindowExtents.Right) / 2;
+				int CenterY = (InCommandData.Application.MainWindowExtents.Top + InCommandData.Application.MainWindowExtents.Bottom) / 2;
+				
+				ExportMessagesDialog = new DatasmithRevitExportMessages(new System.Drawing.Point(CenterX, CenterY), () => ExportMessages = "");
 				ExportMessagesDialog.Messages = ExportMessages;
 				ExportMessagesDialog.Show();
 			}
diff --git a/Engine/Source/Programs/Enterprise/Datasmith/DatasmithRevitExporter/Private/DatasmithRevitCommand.cs b/Engine/Source/Programs/Enterprise/Datasmith/DatasmithRevitExporter/Private/DatasmithRevitCommand.cs
--- a/Engine/Source/Programs/Enterprise/Datasmith/DatasmithRevitExporter/Private/DatasmithRevitCommand.cs
+++ b/Engine/Source/Programs/Enterprise/Datasmith/DatasmithRevitExporter/Private/DatasmithRevitCommand.cs
@@ -401,9 +401,17 @@ namespace DatasmithRevitExporter
 	[Transaction(TransactionMode.Manual)]
 	public class DatasmithManageConnectionsRevitCommand : IExternalCommand
 	{
+		private static bool ConnectionWindowCenterSet = false;
 		public Result Execute(ExternalCommandData InCommandData, ref string OutCommandMessage, ElementSet OutElements) 
 		{
 			IDirectLinkUI DirectLinkUI = IDatasmithExporterUIModule.Get()?.GetDirectLinkExporterUI();
+			if (!ConnectionWindowCenterSet)
+			{
+				int CenterX = (InCommandData.Application.MainWindowExtents.Left + InCommandData.Application.MainWindowExtents.Right) / 2;
+				int CenterY = (InCommandData.Application.MainWindowExtents.Top + InCommandData.Application.MainWindowExtents.Bottom) / 2;
+				DirectLinkUI?.SetStreamWindowCenter(CenterX, CenterY);
+				ConnectionWindowCenterSet = true;
+			}
 			DirectLinkUI?.OpenDirectLinkStreamWindow();
 			return Result.Succeeded;
 		}
@@ -414,7 +422,7 @@ namespace DatasmithRevitExporter
 	{
 		public Result Execute(ExternalCommandData InCommandData, ref string OutCommandMessage, ElementSet OutElements)
 		{
-			DatasmithRevitApplication.ShowExportMessages();
+			DatasmithRevitApplication.ShowExportMessages(InCommandData);
 			return Result.Succeeded;
 		}
 	}
diff --git a/Engine/Source/Programs/Enterprise/Datasmith/DatasmithRevitExporter/Private/DatasmithRevitExportMessages.cs b/Engine/Source/Programs/Enterprise/Datasmith/DatasmithRevitExporter/Private/DatasmithRevitExportMessages.cs
--- a/Engine/Source/Programs/Enterprise/Datasmith/DatasmithRevitExporter/Private/DatasmithRevitExportMessages.cs
+++ b/Engine/Source/Programs/Enterprise/Datasmith/DatasmithRevitExporter/Private/DatasmithRevitExportMessages.cs
@@ -15,9 +15,21 @@ namespace DatasmithRevitExporter
 
 		private ClearCallback OnClear;
 
-		public DatasmithRevitExportMessages(ClearCallback InCallback)
+		private static bool ClosedBefore = false;
+		private static Point ShowLocation = new Point(0, 0);
+		private static Size ShowClientSize = new Size(1200, 300);
+
+		public DatasmithRevitExportMessages(Point InCenter, ClearCallback InCallback)
 		{
+			if (!ClosedBefore)
+			{
+				ShowLocation = InCenter;
+				ShowLocation.X -= (ShowClientSize.Width / 2);
+				ShowLocation.Y -= (ShowClientSize.Height / 2);
+			}
+
 			OnClear = InCallback;
+			FormClosing += OnClosing;
 
 			MessageBox = new TextBox();
 			MessageBox.Name = "MessageBox";
@@ -73,11 +85,13 @@ namespace DatasmithRevitExporter
 			Text = DatasmithRevitResources.Strings.MessagesDialog_Title;
 			AutoScaleDimensions = new SizeF(12F, 25F);
 			AutoScaleMode = AutoScaleMode.Font;
-			ClientSize = new Size(1200, 300);
 			Padding = new Padding(10);
 			ShowIcon = false;
 			SizeGripStyle = SizeGripStyle.Show;
-			StartPosition = FormStartPosition.CenterParent;
+
+			StartPosition = FormStartPosition.Manual;
+			Location = ShowLocation;
+			ClientSize = ShowClientSize;
 
 			Controls.Add(DialogLayout);
 		}
@@ -90,6 +104,13 @@ namespace DatasmithRevitExporter
 			Close();
 		}
 
+		private void OnClosing(object InSender, FormClosingEventArgs InArgs)
+		{
+			ClosedBefore = true;
+			ShowLocation = Location;
+			ShowClientSize = ClientSize;
+		}
+
 		private void ClearButtonClicked(
 			object InSender,
 			EventArgs InEventArgs
diff --git a/Engine/Source/Programs/Enterprise/Datasmith/DatasmithRevitExporter/Private/DatasmithRevitSettingsDialog.cs b/Engine/Source/Programs/Enterprise/Datasmith/DatasmithRevitExporter/Private/DatasmithRevitSettingsDialog.cs
--- a/Engine/Source/Programs/Enterprise/Datasmith/DatasmithRevitExporter/Private/DatasmithRevitSettingsDialog.cs
+++ b/Engine/Source/Programs/Enterprise/Datasmith/DatasmithRevitExporter/Private/DatasmithRevitSettingsDialog.cs
@@ -95,6 +95,9 @@ namespace DatasmithRevitExporter
 		private List<string> AddedParamNames = new List<string>();
 		private bool bSettingsChanged = false;
 
+		private static FormStartPosition ShowStartPosition = FormStartPosition.CenterParent;
+		private static Point ShowLocation = new Point(0, 0);
+
 		public DatasmithRevitSettingsDialog(Autodesk.Revit.DB.Document InDocument, FSettings InSettings)
 		{
 			string FormatTooltip(string InText)
@@ -308,7 +311,9 @@ namespace DatasmithRevitExporter
 			MinimizeBox = false;
 			//Padding = new Padding(7);
 			SizeGripStyle = SizeGripStyle.Hide;
-			StartPosition = FormStartPosition.CenterParent;
+
+			StartPosition = ShowStartPosition;
+			Location = ShowLocation;
 
 			Panel ParentPanel = new Panel();
 			ParentPanel.BorderStyle = BorderStyle.FixedSingle;
@@ -366,6 +371,9 @@ namespace DatasmithRevitExporter
 
 		private void OnClosing(object InSender, FormClosingEventArgs InArgs)
 		{
+			ShowStartPosition = FormStartPosition.Manual;
+			ShowLocation = Location;
+
 			if (bSettingsChanged)
 			{
 				Settings.LevelOfTesselation = Decimal.ToInt32(LevelOfTessellation.Value);
