Datasmith SDK: Cleanup old export implementation

diff --git a/Engine/Source/Developer/Datasmith/DatasmithExporter/Private/DatasmithMeshExporter.cpp b/Engine/Source/Developer/Datasmith/DatasmithExporter/Private/DatasmithMeshExporter.cpp
--- a/Engine/Source/Developer/Datasmith/DatasmithExporter/Private/DatasmithMeshExporter.cpp
+++ b/Engine/Source/Developer/Datasmith/DatasmithExporter/Private/DatasmithMeshExporter.cpp
@@ -16,18 +16,14 @@
 #include "Misc/Paths.h"
 #include "Serialization/MemoryWriter.h"
 
-/**
- * Number of datasmith meshes pending deletion for which the Garbage Collection will be triggered.
- */
-constexpr int32 DSMeshGarbageCollectionThreshold{ 2000 };
 
 struct FDatasmithMeshExporterOptions
 {
-	FDatasmithMeshExporterOptions( const FString& InFullPath, FDatasmithMesh& InMesh, EDSExportLightmapUV InLightmapUV, FDatasmithMesh* InCollisionMesh = nullptr )
-		: MeshFullPath( InFullPath )
-		, Mesh( InMesh )
-		, LightmapUV( InLightmapUV )
-		, CollisionMesh( InCollisionMesh )
+	FDatasmithMeshExporterOptions(const FString& InFullPath, FDatasmithMesh& InMesh, EDSExportLightmapUV InLightmapUV, FDatasmithMesh* InCollisionMesh = nullptr)
+		: MeshFullPath(InFullPath)
+		, Mesh(InMesh)
+		, LightmapUV(InLightmapUV)
+		, CollisionMesh(InCollisionMesh)
 	{}
 
 	FString MeshFullPath;
@@ -36,122 +32,49 @@ struct FDatasmithMeshExporterOptions
 	FDatasmithMesh* CollisionMesh;
 };
 
-/**
- * Implementation class of the DatasmithMeshExporter
- * We use a lockfree UDatasmithMesh pool to avoid creating new UObject when exporting and reduces our memory footprint.
- */
-class FDatasmithMeshExporterLegacyImpl
-{
-public:
-
-	FDatasmithMeshExporterLegacyImpl() :
-		UniqueID(FGuid::NewGuid())
-	{}
-
-	~FDatasmithMeshExporterLegacyImpl()
-	{
-		ClearUDatasmithMeshPool();
-	}
-
-	/**
-	 * Generate and fill a UDatasmithMesh from a FDatasmithMesh.
-	 * Note: The returned TSharedPtr uses a special destructor to return the UDatasmithMesh to a pool instead of destroying it (which wouldn't work anyways).
-	 */
-	TSharedPtr<UDatasmithMesh> GeneratePooledUDatasmithMeshFromFDatasmithMesh( const FDatasmithMesh& Mesh, bool bIsCollisionMesh );
-
-	bool DoExport( TSharedPtr< IDatasmithMeshElement >& MeshElement, const FDatasmithMeshExporterOptions& ExportOptions );
-	void PreExport( const FDatasmithMeshExporterOptions& ExporterOptions );
-	bool ExportMeshes( const FDatasmithMeshExporterOptions& ExporterOptions, FMD5Hash& OutHash );
-	void PostExport(const FDatasmithMesh& DatasmithMesh, TSharedPtr< IDatasmithMeshElement >& MeshElement);
-
-	FString LastError;
-
-private:
-	/**
-	 * This function allows reusing an instanced UDatasmithMesh. Reusing the same object will avoid creating new garbage in memory.
-	 */
-	void FillUDatasmithMeshFromFDatasmithMesh( TSharedPtr<UDatasmithMesh>& UMesh, const FDatasmithMesh& Mesh, bool bValidateRawMesh );
-
-	TSharedPtr<UDatasmithMesh> GetPooledUDatasmithMesh(bool bIsCollisionMesh);
-
-	void ReturnUDatasmithMeshToPool(UDatasmithMesh*& UMesh);
-
-	void ClearUDatasmithMeshPool();
-
-	/**
-	 * A pool of UDatasmithMesh that we use to avoid creating new UObject, this greatly reduce the garbage created when one instance of FDatasmithMeshExporter is used to export multiple Meshes.
-	 */
-	TLockFreePointerListLIFO< UDatasmithMesh > DatasmithMeshUObjectPool;
-	static TAtomic<int32> NumberOfUMeshPendingGC;
-
-	/**
-	 * This exporter's UniqueID. Used to make sure pooled meshes are not using the same names across different threads.
-	 */
-	FGuid UniqueID;
-
-	/**
-	 * Custom deleter used by the TSharedPtr of a pooled UDatasmithMesh to return it to the pool instead of deleting them.
-	 */
-	struct FPooledUDatasmithMeshDeleter
-	{
-		FPooledUDatasmithMeshDeleter(FDatasmithMeshExporterLegacyImpl* InExporterPtr)
-			: ExporterPtr(InExporterPtr)
-		{}
-
-		FORCEINLINE void operator()(UDatasmithMesh* Object) const
-		{
-			check(Object != nullptr && ExporterPtr != nullptr);
-			ExporterPtr->ReturnUDatasmithMeshToPool(Object);
-		}
-
-		FDatasmithMeshExporterLegacyImpl* ExporterPtr;
-	};
-};
-
-TAtomic<int32> FDatasmithMeshExporterLegacyImpl::NumberOfUMeshPendingGC(0);
-
 
 class FDatasmithMeshExporterImpl
 {
 public:
-	bool DoExport( TSharedPtr< IDatasmithMeshElement >& MeshElement, const FDatasmithMeshExporterOptions& ExportOptions );
+	bool DoExport(TSharedPtr<IDatasmithMeshElement>& MeshElement, const FDatasmithMeshExporterOptions& ExportOptions);
 	FString LastError;
 
 private:
 	bool WriteMeshFile(const FDatasmithMeshExporterOptions& ExporterOptions, FMD5Hash& OutHash);
 };
 
-bool FDatasmithMeshExporterImpl::DoExport(TSharedPtr< IDatasmithMeshElement >& MeshElement, const FDatasmithMeshExporterOptions& ExportOptions)
+
+bool FDatasmithMeshExporterImpl::DoExport(TSharedPtr<IDatasmithMeshElement>& MeshElement, const FDatasmithMeshExporterOptions& ExportOptions)
 {
 	FDatasmithMesh& Mesh = ExportOptions.Mesh;
 
 	// If the mesh doesn't have a name, use the filename as its name
-	if ( FCString::Strlen( Mesh.GetName() ) == 0 )
+	if (FCString::Strlen(Mesh.GetName()) == 0)
 	{
-		Mesh.SetName( *FPaths::GetBaseFilename( ExportOptions.MeshFullPath ) );
+		Mesh.SetName(*FPaths::GetBaseFilename(ExportOptions.MeshFullPath));
 	}
 
 	// make sure we have at least one UV channels on each source models
-	FDatasmithMeshUtils::CreateDefaultUVsWithLOD( Mesh );
+	FDatasmithMeshUtils::CreateDefaultUVsWithLOD(Mesh);
 
 	FMD5Hash Hash;
-	if ( WriteMeshFile( ExportOptions, Hash ) )
+	if (WriteMeshFile(ExportOptions, Hash))
 	{
 		// If no existing MeshElement provided, create one.
-		if ( !MeshElement )
+		if (!MeshElement)
 		{
-			MeshElement = FDatasmithSceneFactory::CreateMesh( Mesh.GetName() );
+			MeshElement = FDatasmithSceneFactory::CreateMesh(Mesh.GetName());
 		}
-		MeshElement->SetFile( *ExportOptions.MeshFullPath );
-		MeshElement->SetFileHash( Hash );
+		MeshElement->SetFile(*ExportOptions.MeshFullPath);
+		MeshElement->SetFileHash(Hash);
 
 		FBox3f Extents = Mesh.GetExtents();
 		float Width = Extents.Max[0] - Extents.Min[0];
 		float Height = Extents.Max[2] - Extents.Min[2];
 		float Depth = Extents.Max[1] - Extents.Min[1];
 
-		MeshElement->SetDimensions( Mesh.ComputeArea(), Width, Height, Depth );
-		MeshElement->SetLightmapSourceUV( Mesh.GetLightmapSourceUVChannel() );
+		MeshElement->SetDimensions(Mesh.ComputeArea(), Width, Height, Depth);
+		MeshElement->SetLightmapSourceUV(Mesh.GetLightmapSourceUVChannel());
 
 		return true;
 	}
@@ -164,82 +87,28 @@ bool FDatasmithMeshExporterImpl::DoExport(TSharedPtr< IDatasmithMeshElement >& M
 }
 
 
-bool FDatasmithMeshExporterLegacyImpl::ExportMeshes( const FDatasmithMeshExporterOptions& ExporterOptions, FMD5Hash& OutHash )
-{
-	TArray<TSharedPtr<UDatasmithMesh>, TInlineAllocator<2>> MeshesToExport;
-
-	// Static mesh, we keep a static UDatasmithMesh alive as a utility object and re-use it for every export instead of creating a new one every time. This avoid creating garbage in memory.
-	bool bIsCollisionMesh = false;
-	MeshesToExport.Add( GeneratePooledUDatasmithMeshFromFDatasmithMesh( ExporterOptions.Mesh, bIsCollisionMesh ) );
-
-	// Collision mesh
-	if ( ExporterOptions.CollisionMesh )
-	{
-		bIsCollisionMesh = true;
-		MeshesToExport.Add( GeneratePooledUDatasmithMeshFromFDatasmithMesh( *ExporterOptions.CollisionMesh, bIsCollisionMesh ) );
-	}
-
-	TUniquePtr<FArchive> Archive( IFileManager::Get().CreateFileWriter( *ExporterOptions.MeshFullPath ) );
-
-	if ( !Archive.IsValid() )
-	{
-		LastError = FString::Printf( TEXT( "Failed writing to file %s" ), *ExporterOptions.MeshFullPath );
 
-		return false;
-	}
-
-	int32 NumMeshes = MeshesToExport.Num();
-
-	*Archive << NumMeshes;
-
-	FMD5 MD5;
-	for ( TSharedPtr<UDatasmithMesh>& MeshToExport : MeshesToExport )
-	{
-		TArray< uint8 > Bytes;
-		FMemoryWriter MemoryWriter( Bytes, true );
-		MemoryWriter.ArIgnoreClassRef = false;
-		MemoryWriter.ArIgnoreArchetypeRef = false;
-		MemoryWriter.ArNoDelta = false;
-		MemoryWriter.SetWantBinaryPropertySerialization( true );
-
-		MeshToExport->Serialize( MemoryWriter );
-
-		// Calculate the Hash of all the mesh to export
-		for ( FDatasmithMeshSourceModel& Model : MeshToExport->SourceModels )
-		{
-			uint8* Buffer = (uint8*)Model.RawMeshBulkData.GetBulkData().LockReadOnly();
-			MD5.Update( Buffer, Model.RawMeshBulkData.GetBulkData().GetBulkDataSize() );
-			Model.RawMeshBulkData.GetBulkData().Unlock();
-		}
-
-		*Archive << Bytes;
-	}
-	OutHash.Set( MD5 );
-	return true;
-}
-
-
-FDatasmithMeshExporter::FDatasmithMeshExporter(bool bAllowOldSerialization)
-	: LegacyImpl(MakeUnique<FDatasmithMeshExporterLegacyImpl>())
-	, Impl(MakeUnique<FDatasmithMeshExporterImpl>())
-	, bAllowOldSerialization(bAllowOldSerialization)
+FDatasmithMeshExporter::FDatasmithMeshExporter()
+	: Impl(MakeUnique<FDatasmithMeshExporterImpl>())
 {}
 
+
 FDatasmithMeshExporter::~FDatasmithMeshExporter() = default;
 
-FString GetMeshFilePath( const TCHAR* Filepath, const TCHAR* Filename )
+
+FString GetMeshFilePath(const TCHAR* Filepath, const TCHAR* Filename)
 {
 	FString NormalizedFilepath = Filepath;
-	FPaths::NormalizeDirectoryName( NormalizedFilepath );
+	FPaths::NormalizeDirectoryName(NormalizedFilepath);
 
 	FString NormalizedFilename = Filename;
-	FPaths::NormalizeFilename( NormalizedFilename );
+	FPaths::NormalizeFilename(NormalizedFilename);
 
-	return FPaths::Combine( *NormalizedFilepath, FPaths::SetExtension( NormalizedFilename, UDatasmithMesh::GetFileExtension() ) );
+	return FPaths::Combine(*NormalizedFilepath, FPaths::SetExtension(NormalizedFilename, UDatasmithMesh::GetFileExtension()));
 }
 
 
-bool FDatasmithMeshExporterImpl::WriteMeshFile( const FDatasmithMeshExporterOptions& ExporterOptions, FMD5Hash& OutHash )
+bool FDatasmithMeshExporterImpl::WriteMeshFile(const FDatasmithMeshExporterOptions& ExporterOptions, FMD5Hash& OutHash)
 {
 	FDatasmithPackedMeshes Pack;
 
@@ -271,10 +140,10 @@ bool FDatasmithMeshExporterImpl::WriteMeshFile( const FDatasmithMeshExporterOpti
 		Pack.Meshes.Add(PackMeshModels(*ExporterOptions.CollisionMesh, true));
 	}
 
-	TUniquePtr<FArchive> Archive( IFileManager::Get().CreateFileWriter(*ExporterOptions.MeshFullPath) );
+	TUniquePtr<FArchive> Archive(IFileManager::Get().CreateFileWriter(*ExporterOptions.MeshFullPath));
 	if (!Archive.IsValid())
 	{
-		LastError = FString::Printf( TEXT( "Failed writing to file %s" ), *ExporterOptions.MeshFullPath );
+		LastError = FString::Printf(TEXT("Failed writing to file %s"), *ExporterOptions.MeshFullPath);
 		return false;
 	}
 
@@ -286,216 +155,29 @@ bool FDatasmithMeshExporterImpl::WriteMeshFile( const FDatasmithMeshExporterOpti
 }
 
 
-TSharedPtr< IDatasmithMeshElement > FDatasmithMeshExporter::ExportToUObject(const TCHAR* Filepath, const TCHAR* Filename, FDatasmithMesh& Mesh, FDatasmithMesh* CollisionMesh, EDSExportLightmapUV LightmapUV, bool NewFormat)
+TSharedPtr<IDatasmithMeshElement> FDatasmithMeshExporter::ExportToUObject(const TCHAR* Filepath, const TCHAR* Filename, FDatasmithMesh& Mesh, FDatasmithMesh* CollisionMesh, EDSExportLightmapUV LightmapUV)
 {
-	FString FullPath( GetMeshFilePath( Filepath, Filename ) );
+	FString FullPath(GetMeshFilePath(Filepath, Filename));
 
 	TSharedPtr<IDatasmithMeshElement> ExportedMeshElement;
-	FDatasmithMeshExporterOptions ExportOptions( MoveTemp( FullPath ), Mesh, LightmapUV, CollisionMesh );
-	if (NewFormat)
-	{
-		Impl->DoExport( ExportedMeshElement, ExportOptions );
-	}
-	else
-	{
-		ensureMsgf(bAllowOldSerialization, TEXT("Old serialization used for mesh %s. It is no longer supported."), *FullPath);
-		LegacyImpl->DoExport( ExportedMeshElement, ExportOptions );
-	}
+	FDatasmithMeshExporterOptions ExportOptions(MoveTemp(FullPath), Mesh, LightmapUV, CollisionMesh);
+	Impl->DoExport(ExportedMeshElement, ExportOptions);
 
 	return ExportedMeshElement;
 }
 
 
-bool FDatasmithMeshExporter::ExportToUObject( TSharedPtr< IDatasmithMeshElement >& MeshElement, const TCHAR* Filepath, FDatasmithMesh& Mesh, FDatasmithMesh* CollisionMesh, EDSExportLightmapUV LightmapUV, bool NewFormat )
+bool FDatasmithMeshExporter::ExportToUObject(TSharedPtr<IDatasmithMeshElement>& MeshElement, const TCHAR* Filepath, FDatasmithMesh& Mesh, FDatasmithMesh* CollisionMesh, EDSExportLightmapUV LightmapUV)
 {
-	FString FullPath( GetMeshFilePath( Filepath, MeshElement->GetName() ) );
-	FDatasmithMeshExporterOptions ExportOptions( FullPath, Mesh, LightmapUV, CollisionMesh );
+	FString FullPath(GetMeshFilePath(Filepath, MeshElement->GetName()));
+	FDatasmithMeshExporterOptions ExportOptions(FullPath, Mesh, LightmapUV, CollisionMesh);
 
-	if (NewFormat)
-	{
-		return Impl->DoExport( MeshElement, ExportOptions );
-	}
-	else
-	{
-		ensureMsgf(bAllowOldSerialization, TEXT("Old serialization used for mesh %s. It is no longer supported."), MeshElement->GetName());
-		return LegacyImpl->DoExport( MeshElement, ExportOptions );
-	}
+	return Impl->DoExport(MeshElement, ExportOptions);
 }
 
 
 FString FDatasmithMeshExporter::GetLastError() const
 {
-	return LegacyImpl->LastError.IsEmpty() ? Impl->LastError : LegacyImpl->LastError;
+	return Impl->LastError;
 }
 
-bool FDatasmithMeshExporterLegacyImpl::DoExport( TSharedPtr< IDatasmithMeshElement >& MeshElement, const FDatasmithMeshExporterOptions& ExportOptions )
-{
-	FMD5Hash Hash;
-	PreExport( ExportOptions );
-	if ( ExportMeshes( ExportOptions, Hash ) )
-	{
-		// If no existing MeshElement provided, create one.
-		if ( !MeshElement )
-		{
-			FString BaseFileName = FPaths::GetBaseFilename( ExportOptions.MeshFullPath );
-			MeshElement = FDatasmithSceneFactory::CreateMesh( *BaseFileName );
-		}
-		MeshElement->SetFile( *ExportOptions.MeshFullPath );
-		MeshElement->SetFileHash( Hash );
-
-		PostExport( ExportOptions.Mesh, MeshElement );
-
-		return true;
-	}
-
-	return false;
-}
-
-void FDatasmithMeshExporterLegacyImpl::PreExport( const FDatasmithMeshExporterOptions& ExporterOptions )
-{
-	FDatasmithMesh& Mesh = ExporterOptions.Mesh;
-
-	// If the mesh doesn't have a name, use the filename as its name
-	if ( FCString::Strlen( Mesh.GetName() ) == 0 )
-	{
-		Mesh.SetName( *FPaths::GetBaseFilename( ExporterOptions.MeshFullPath ) );
-	}
-
-	FDatasmithMeshUtils::CreateDefaultUVsWithLOD( Mesh );
-}
-
-void FDatasmithMeshExporterLegacyImpl::PostExport( const FDatasmithMesh& DatasmithMesh, TSharedPtr< IDatasmithMeshElement >& MeshElement )
-{
-	FBox3f Extents = DatasmithMesh.GetExtents();
-	float Width = Extents.Max[0] - Extents.Min[0];
-	float Height = Extents.Max[2] - Extents.Min[2];
-	float Depth = Extents.Max[1] - Extents.Min[1];
-
-	MeshElement->SetDimensions( DatasmithMesh.ComputeArea(), Width, Height, Depth );
-
-	MeshElement->SetLightmapSourceUV( DatasmithMesh.GetLightmapSourceUVChannel() );
-}
-
-TSharedPtr<UDatasmithMesh> FDatasmithMeshExporterLegacyImpl::GeneratePooledUDatasmithMeshFromFDatasmithMesh(const FDatasmithMesh& Mesh, bool bIsCollisionMesh)
-{
-	TSharedPtr<UDatasmithMesh> PooledMesh = GetPooledUDatasmithMesh(bIsCollisionMesh);
-	FillUDatasmithMeshFromFDatasmithMesh(PooledMesh, Mesh, !bIsCollisionMesh);
-	PooledMesh->bIsCollisionMesh = bIsCollisionMesh;
-	return PooledMesh;
-}
-
-/**
- * This function allows reusing an instanced UDatasmithMesh. Reusing the same object will avoid creating new garbage in memory.
- */
-void FDatasmithMeshExporterLegacyImpl::FillUDatasmithMeshFromFDatasmithMesh(TSharedPtr<UDatasmithMesh>& UMesh, const FDatasmithMesh& Mesh, bool bValidateRawMesh)
-{
-	UMesh->MeshName = Mesh.GetName();
-
-	FRawMesh RawMesh;
-	FDatasmithMeshUtils::ToRawMesh(Mesh, RawMesh, bValidateRawMesh);
-
-	FDatasmithMeshSourceModel BaseModel;
-	BaseModel.RawMeshBulkData.SaveRawMesh(RawMesh);
-
-	UMesh->SourceModels.Add(BaseModel);
-
-	for (int32 LODIndex = 0; LODIndex < Mesh.GetLODsCount(); ++LODIndex)
-	{
-		if (const FDatasmithMesh* LODMesh = Mesh.GetLOD(LODIndex))
-		{
-			FDatasmithMeshUtils::ToRawMesh(*LODMesh, RawMesh, bValidateRawMesh);
-
-			FRawMeshBulkData LODRawMeshBulkData;
-			LODRawMeshBulkData.SaveRawMesh(RawMesh);
-
-			FDatasmithMeshSourceModel LODModel;
-			LODModel.RawMeshBulkData = LODRawMeshBulkData;
-
-			UMesh->SourceModels.Add(LODModel);
-		}
-	}
-}
-
-TSharedPtr<UDatasmithMesh> FDatasmithMeshExporterLegacyImpl::GetPooledUDatasmithMesh(bool bIsCollisionMesh)
-{
-	const FString GuidString = UniqueID.ToString();
-	const FString CollisionString = bIsCollisionMesh ? TEXT("_Collision") : TEXT("");
-	const FString PooledMeshName = FString::Printf(TEXT("DatasmithExporter_%s_TransientPooledUDatasmithMesh%s"), *GuidString, *CollisionString);
-
-	if (UDatasmithMesh* PooledMesh = DatasmithMeshUObjectPool.Pop())
-	{
-		PooledMesh->Rename(*PooledMeshName);
-		return TSharedPtr<UDatasmithMesh>(PooledMesh, FPooledUDatasmithMeshDeleter(this));
-	}
-
-	{
-		// Can't create new objects while the GC is running.
-		FGCScopeGuard GCGuard;
-		return TSharedPtr<UDatasmithMesh>(NewObject< UDatasmithMesh >((UObject*)GetTransientPackage(), *PooledMeshName, RF_Transient | RF_MarkAsRootSet), FPooledUDatasmithMeshDeleter(this));
-	}
-}
-
-void FDatasmithMeshExporterLegacyImpl::ReturnUDatasmithMeshToPool(UDatasmithMesh*& UMesh)
-{
-	//Clear the UDatasmithMesh.
-	UMesh->SourceModels.Empty();
-
-	//Put it back into the pool
-	DatasmithMeshUObjectPool.Push(UMesh);
-	//Null the reference to make sure we don't use the object returned to the pool.
-	UMesh = nullptr;
-}
-
-void FDatasmithMeshExporterLegacyImpl::ClearUDatasmithMeshPool()
-{
-	TArray<UDatasmithMesh*> PooledMeshes;
-	DatasmithMeshUObjectPool.PopAll(PooledMeshes);
-
-	//Keep track of the number of garbage UObject generated by clearing the cache so that we can trigger the GC after a while.
-	NumberOfUMeshPendingGC += PooledMeshes.Num();
-
-	auto ClearUObjectFlags = [PooledMeshes(MoveTemp(PooledMeshes))]()
-	{
-		for (UDatasmithMesh* UMesh : PooledMeshes)
-		{
-			UMesh->RemoveFromRoot();
-			UMesh->MarkAsGarbage();
-
-			//This object won't be reused, we must make sure there is no Async flag so that the GC knows it can safely collect it.
-			UMesh->ClearInternalFlags(EInternalObjectFlags::Async);
-		}
-	};
-
-#if IS_PROGRAM
-	if (FDatasmithExporterManager::WasInitializedWithGameThread())
-	{
-		FSimpleDelegate Command;
-		Command.BindLambda(MoveTemp(ClearUObjectFlags));
-
-		// No need to synchronize to the game thread for this, we can lazily update the object flags by pushing a command, that way the flags will be updated before the next GC without blocking the current thread.
-		const bool bWakeUpGameThread = false;
-		FDatasmithExporterManager::PushCommandIntoGameThread(MoveTemp(Command), bWakeUpGameThread);
-	}
-	else
-#endif
-	if (IsInGameThread())
-	{
-		ClearUObjectFlags();
-	}
-	else
-	{
-		// Can't modify the GC flags while it's running.
-		FGCScopeGuard GCGuard;
-		ClearUObjectFlags();
-	}
-
-	//Even if our UObjects are basically empty at this point and don't have a big memory footprint, the engine will assert when reaching around 131k UObjects.
-	//So we need to call the GC before that point.
-	if (NumberOfUMeshPendingGC > DSMeshGarbageCollectionThreshold)
-	{
-		if (FDatasmithExporterManager::RunGarbageCollection())
-		{
-			NumberOfUMeshPendingGC = 0;
-		}
-	}
-}
diff --git a/Engine/Source/Developer/Datasmith/DatasmithExporter/Public/DatasmithMeshExporter.h b/Engine/Source/Developer/Datasmith/DatasmithExporter/Public/DatasmithMeshExporter.h
--- a/Engine/Source/Developer/Datasmith/DatasmithExporter/Public/DatasmithMeshExporter.h
+++ b/Engine/Source/Developer/Datasmith/DatasmithExporter/Public/DatasmithMeshExporter.h
@@ -13,8 +13,7 @@ class IDatasmithMeshElement;
 class DATASMITHEXPORTER_API FDatasmithMeshExporter
 {
 public:
-	FDatasmithMeshExporter(bool bAllowOldSerialization);
-	FDatasmithMeshExporter() : FDatasmithMeshExporter(false) {}
+	FDatasmithMeshExporter();
 	virtual ~FDatasmithMeshExporter();
 
 	/**
@@ -28,7 +27,7 @@ public:
 	 *
 	 * @return				A IDatasmithMeshElement that refers to the exported file
 	 */
-	TSharedPtr< IDatasmithMeshElement > ExportToUObject( const TCHAR* Filepath, const TCHAR* Filename, FDatasmithMesh& Mesh, FDatasmithMesh* CollisionMesh, EDSExportLightmapUV LightmapUV, bool NewFormat=true );
+	TSharedPtr<IDatasmithMeshElement> ExportToUObject(const TCHAR* Filepath, const TCHAR* Filename, FDatasmithMesh& Mesh, FDatasmithMesh* CollisionMesh, EDSExportLightmapUV LightmapUV);
 
 	/**
 	 * Exports a FDatasmithMesh as a UObject and link it to the given IDatasmithMeshElementElement.
@@ -41,7 +40,7 @@ public:
 	 *
 	 * @return				True if export was successful.
 	 */
-	bool ExportToUObject( TSharedPtr< IDatasmithMeshElement >& MeshElement, const TCHAR* Filepath, FDatasmithMesh& Mesh, FDatasmithMesh* CollisionMesh, EDSExportLightmapUV LightmapUV, bool NewFormat=true );
+	bool ExportToUObject(TSharedPtr<IDatasmithMeshElement>& MeshElement, const TCHAR* Filepath, FDatasmithMesh& Mesh, FDatasmithMesh* CollisionMesh, EDSExportLightmapUV LightmapUV);
 
 	/**
 	 * @return The error that happened during the last export, if any
@@ -49,7 +48,6 @@ public:
 	FString GetLastError() const;
 
 private:
-	TUniquePtr<class FDatasmithMeshExporterLegacyImpl> LegacyImpl;
 	TUniquePtr<class FDatasmithMeshExporterImpl> Impl;
 	bool bAllowOldSerialization = false;
 };
diff --git a/Engine/Source/Runtime/Datasmith/DatasmithCore/Private/DatasmithMeshSerialization.cpp b/Engine/Source/Runtime/Datasmith/DatasmithCore/Private/DatasmithMeshSerialization.cpp
--- a/Engine/Source/Runtime/Datasmith/DatasmithCore/Private/DatasmithMeshSerialization.cpp
+++ b/Engine/Source/Runtime/Datasmith/DatasmithCore/Private/DatasmithMeshSerialization.cpp
@@ -243,6 +243,7 @@ TOptional<FMeshDescription> ExtractToMeshDescription(FDatasmithMeshSourceModel&
 	return MeshDescription;
 }
 
+// Reads previous format of meshes (RawMesh based)
 TArray<FDatasmithMeshModels> GetDatasmithMeshFromMeshPath_Legacy(FArchive* Archive, int32 LeagacyNumMeshesCount)
 {
 	TArray< FDatasmithMeshModels > Result;
@@ -317,7 +318,7 @@ FDatasmithPackedMeshes GetDatasmithMeshFromFile(const FString& MeshPath)
 	return Result;
 }
 
-DATASMITHCORE_API FDatasmithPackedMeshes GetDatasmithClothFromFile(const FString& Path)
+FDatasmithPackedMeshes GetDatasmithClothFromFile(const FString& Path)
 {
 	// #ue_ds_cloth_note: asset serialization
 	FDatasmithPackedMeshes Result;
