Datasmith: Log the cause of a serialization issue

diff --git a/Engine/Source/Developer/Datasmith/DatasmithExporter/Private/DatasmithMeshExporter.cpp b/Engine/Source/Developer/Datasmith/DatasmithExporter/Private/DatasmithMeshExporter.cpp
--- a/Engine/Source/Developer/Datasmith/DatasmithExporter/Private/DatasmithMeshExporter.cpp
+++ b/Engine/Source/Developer/Datasmith/DatasmithExporter/Private/DatasmithMeshExporter.cpp
@@ -76,7 +76,7 @@ bool FDatasmithMeshExporterImpl::DoExport(TSharedPtr<IDatasmithMeshElement>& Mes
 	}
 	else
 	{
-		UE_LOG(LogDatasmith, Warning, TEXT("Cannot export mesh %s"), Mesh.GetName());
+		UE_LOG(LogDatasmith, Warning, TEXT("Cannot export mesh %s: %s"), Mesh.GetName(), *LastError);
 	}
 
 	return false;
