DatasmithFacade: remove unacceptable words

diff --git a/Engine/Source/Developer/Datasmith/DatasmithFacade/Private/DatasmithFacadeDecal.cpp b/Engine/Source/Developer/Datasmith/DatasmithFacade/Private/DatasmithFacadeDecal.cpp
--- a/Engine/Source/Developer/Datasmith/DatasmithFacade/Private/DatasmithFacadeDecal.cpp
+++ b/Engine/Source/Developer/Datasmith/DatasmithFacade/Private/DatasmithFacadeDecal.cpp
@@ -55,8 +55,8 @@ void FDatasmithFacadeActorDecal::SetSortOrder(int32 InSortOrder)
 // class FDatasmithFacadeDecalMaterial
 
 FDatasmithFacadeDecalMaterial::FDatasmithFacadeDecalMaterial(
-	const TSharedRef<IDatasmithDecalMaterialElement>& InMaterialRef // Datasmith master material element
-)	
+	const TSharedRef<IDatasmithDecalMaterialElement>& InMaterialRef
+)
 	: FDatasmithFacadeBaseMaterial(InMaterialRef)
 {}
 
diff --git a/Engine/Source/Developer/Datasmith/DatasmithFacade/Public/DatasmithFacadeDecal.h b/Engine/Source/Developer/Datasmith/DatasmithFacade/Public/DatasmithFacadeDecal.h
--- a/Engine/Source/Developer/Datasmith/DatasmithFacade/Public/DatasmithFacadeDecal.h
+++ b/Engine/Source/Developer/Datasmith/DatasmithFacade/Public/DatasmithFacadeDecal.h
@@ -74,7 +74,7 @@ protected:
 #endif
 
 	FDatasmithFacadeDecalMaterial(
-		const TSharedRef<IDatasmithDecalMaterialElement>& InMaterialRef // Datasmith master material element
+		const TSharedRef<IDatasmithDecalMaterialElement>& InMaterialRef
 	);
 
 	TSharedRef<IDatasmithDecalMaterialElement> GetDatasmithDecalMaterial() const;
diff --git a/Engine/Source/Developer/Datasmith/DatasmithFacade/Public/DatasmithFacadeMaterial.h b/Engine/Source/Developer/Datasmith/DatasmithFacade/Public/DatasmithFacadeMaterial.h
--- a/Engine/Source/Developer/Datasmith/DatasmithFacade/Public/DatasmithFacadeMaterial.h
+++ b/Engine/Source/Developer/Datasmith/DatasmithFacade/Public/DatasmithFacadeMaterial.h
@@ -26,7 +26,7 @@ public:
 protected:
 #endif
 	FDatasmithFacadeBaseMaterial(
-		const TSharedRef<IDatasmithBaseMaterialElement>& BaseMaterialElement 
+		const TSharedRef<IDatasmithBaseMaterialElement>& BaseMaterialElement
 	);
 
 	static EDatasmithMaterialType GetDatasmithMaterialType(
@@ -47,14 +47,13 @@ class DATASMITHFACADE_API FDatasmithFacadeMaterialInstance :
 
 public:
 
-	// Possible Datasmith master material types, from EDatasmithMaterialInstanceType in DatasmithDefinitions.h
+	// Possible Datasmith material types, from EDatasmithMaterialInstanceType in DatasmithDefinitions.h
 	enum class EMaterialInstanceType
 	{
 		Auto,
 		Opaque,
 		Transparent,
 		ClearCoat,
-		/** Instantiate a master material from a specified one */
 		Custom,
 		/** Material has a transparent cutout map */
 		CutOut,
@@ -75,31 +74,27 @@ public:
 public:
 
 	FDatasmithFacadeMaterialInstance(
-		const TCHAR* InElementName // Datasmith element name
+		const TCHAR* InElementName
 	);
 
 	virtual ~FDatasmithFacadeMaterialInstance() {}
 
-	/** Get the Datasmith master material type. */
 	EMaterialInstanceType GetMaterialType() const;
 
-	/** Set the Datasmith master material type. */
 	void SetMaterialType(
-		EMaterialInstanceType InMaterialInstanceType // master material type
+		EMaterialInstanceType InMaterialInstanceType
 	);
 
-	/** Get the Datasmith master material quality. */
 	EMaterialInstanceQuality GetQuality() const;
 
-	/** Set the Datasmith master material quality. */
 	void SetQuality(
 		EMaterialInstanceQuality InQuality
 	);
 
-	/** Get the material path name used when master material type is set to Custom */
+	/** Get the material path name used when material type is set to Custom */
 	const TCHAR* GetCustomMaterialPathName() const;
-	
-	/** Set the material path name used when master material type is set to Custom */
+
+	/** Set the material path name used when material type is set to Custom */
 	void SetCustomMaterialPathName(
 		const TCHAR* InPathName
 	);
@@ -163,7 +158,7 @@ protected:
 #endif
 
 	FDatasmithFacadeMaterialInstance(
-		const TSharedRef<IDatasmithMaterialInstanceElement>& InMaterialRef // Datasmith master material element
+		const TSharedRef<IDatasmithMaterialInstanceElement>& InMaterialRef
 	);
 
 	TSharedRef<IDatasmithMaterialInstanceElement> GetDatasmithMaterialInstance() const;
