Datasmith: Fixed build of Datasmith external plugins/SDK (after CL 21915314 has broken it)

diff --git a/Engine/Source/Developer/Datasmith/DatasmithExporter/DatasmithExporter.Build.cs b/Engine/Source/Developer/Datasmith/DatasmithExporter/DatasmithExporter.Build.cs
--- a/Engine/Source/Developer/Datasmith/DatasmithExporter/DatasmithExporter.Build.cs
+++ b/Engine/Source/Developer/Datasmith/DatasmithExporter/DatasmithExporter.Build.cs
@@ -35,12 +35,8 @@ namespace UnrealBuildTool.Rules
 				}
 			);
 
-			PrivateIncludePaths.AddRange(
-				new string[]
-				{
-					"Runtime/Launch/Public"
-				}
-			);
+			PrivateIncludePaths.Add("Runtime/Launch/Public");
+			PrivateIncludePaths.Add("Runtime/Launch/Private");      // For LaunchEngineLoop.cpp include
 
 			// PrecompileForTargets = PrecompileTargetsType.Any;
 		}
