[Datasmith]
* Fixed non unity/pch compile error happening in ue main

diff --git a/Engine/Source/Developer/Datasmith/DatasmithFacade/Private/DatasmithFacadeActorCamera.cpp b/Engine/Source/Developer/Datasmith/DatasmithFacade/Private/DatasmithFacadeActorCamera.cpp
--- a/Engine/Source/Developer/Datasmith/DatasmithFacade/Private/DatasmithFacadeActorCamera.cpp
+++ b/Engine/Source/Developer/Datasmith/DatasmithFacade/Private/DatasmithFacadeActorCamera.cpp
@@ -1,8 +1,8 @@
 // Copyright Epic Games, Inc. All Rights Reserved.
 
 #include "DatasmithFacadeActorCamera.h"
-
 #include "IDatasmithSceneElements.h"
+#include "Math/RotationMatrix.h"
 
 FDatasmithFacadeActorCamera::FDatasmithFacadeActorCamera(
 	const TCHAR* InElementName
diff --git a/Engine/Source/Developer/Datasmith/DatasmithFacade/Private/DatasmithFacadeVariant.cpp b/Engine/Source/Developer/Datasmith/DatasmithFacade/Private/DatasmithFacadeVariant.cpp
--- a/Engine/Source/Developer/Datasmith/DatasmithFacade/Private/DatasmithFacadeVariant.cpp
+++ b/Engine/Source/Developer/Datasmith/DatasmithFacade/Private/DatasmithFacadeVariant.cpp
@@ -2,6 +2,7 @@
 
 #include "DatasmithFacadeVariant.h"
 #include "DatasmithFacadeMaterial.h"
+#include "Math/Transform.h"
 
 FDatasmithFacadePropertyCapture::FDatasmithFacadePropertyCapture()
 : FDatasmithFacadeElement(FDatasmithSceneFactory::CreatePropertyCapture())
