Datasmith direct link not showing any sources

diff --git a/Engine/Plugins/Messaging/UdpMessaging/Source/UdpMessaging/Private/Transport/UdpMessageProcessor.cpp b/Engine/Plugins/Messaging/UdpMessaging/Source/UdpMessaging/Private/Transport/UdpMessageProcessor.cpp
--- a/Engine/Plugins/Messaging/UdpMessaging/Source/UdpMessaging/Private/Transport/UdpMessageProcessor.cpp
+++ b/Engine/Plugins/Messaging/UdpMessaging/Source/UdpMessaging/Private/Transport/UdpMessageProcessor.cpp
@@ -132,7 +132,10 @@ FUdpMessageProcessor::FUdpMessageProcessor(FSocket& InSocket, const FGuid& InNod
 FUdpMessageProcessor::~FUdpMessageProcessor()
 {
 	// shut down worker thread if it is still running
-	Thread->Kill();
+	if (Thread)
+	{
+		Thread->Kill();
+	}
 
 	Thread = {};
 	Beacon = {};
diff --git a/Engine/Source/Developer/Datasmith/DatasmithExporter/Private/DatasmithExporterManager.cpp b/Engine/Source/Developer/Datasmith/DatasmithExporter/Private/DatasmithExporterManager.cpp
--- a/Engine/Source/Developer/Datasmith/DatasmithExporter/Private/DatasmithExporterManager.cpp
+++ b/Engine/Source/Developer/Datasmith/DatasmithExporter/Private/DatasmithExporterManager.cpp
@@ -96,6 +96,10 @@ namespace DatasmithGameThread
 	{
 		bool bEngineInitialized = GEngineLoop.PreInit(*PreInitCommandArgs) == 0;
 
+		// Workaround for UdpMessaging not starting processing until PostDefault stage - UE-179092
+		// And since Datasmith plugins are built without WITH_ENGINE FEngineLoop::LoadStartupModules is not called which effectively skips PostDefault stage
+		IPluginManager::Get().LoadModulesForEnabledPlugins(ELoadingPhase::PostDefault);
+
 		// Make sure all UObject classes are registered and default properties have been initialized
 		ProcessNewlyLoadedUObjects();
 
