[Datasmith] Fix compilation error caused by implicit pointer validation.

diff --git a/Engine/Source/Developer/Datasmith/DatasmithExporter/Private/DatasmithExporterManager.cpp b/Engine/Source/Developer/Datasmith/DatasmithExporter/Private/DatasmithExporterManager.cpp
--- a/Engine/Source/Developer/Datasmith/DatasmithExporter/Private/DatasmithExporterManager.cpp
+++ b/Engine/Source/Developer/Datasmith/DatasmithExporter/Private/DatasmithExporterManager.cpp
@@ -327,7 +327,7 @@ bool FDatasmithExporterManager::Initialize(const FInitOptions& InitOptions)
 
 		if (InitOptions.bUseDatasmithExporterUI)
 		{
-			checkf(InitOptions.RemoteEngineDirPath, TEXT("Datasmith exporter UI need a path to its minimal engine folder"));
+			checkf(InitOptions.RemoteEngineDirPath != nullptr, TEXT("Datasmith exporter UI need a path to its minimal engine folder"));
 
 			// Start a custom game thread
 			GDatasmithGameThread = MakeShared<FDatasmithGameThread>(MoveTemp(CmdLine), InitOptions);
