Fixed issue with ForeignEngineDir not terminated with a '/'
#rnx#preflight 65c65e237d256b55f352e49c

diff --git a/Engine/Source/Developer/Datasmith/DatasmithExporter/Private/DatasmithExporterManager.cpp b/Engine/Source/Developer/Datasmith/DatasmithExporter/Private/DatasmithExporterManager.cpp
--- a/Engine/Source/Developer/Datasmith/DatasmithExporter/Private/DatasmithExporterManager.cpp
+++ b/Engine/Source/Developer/Datasmith/DatasmithExporter/Private/DatasmithExporterManager.cpp
@@ -238,6 +238,13 @@ void FDatasmithGameThread::RequestExit()
 
 void FDatasmithGameThread::OnInit()
 {
+	// Make sure that external engine path is normalized and ends with a '/'
+	FPaths::NormalizeDirectoryName(ForeignEngineDir);
+	if (!ForeignEngineDir.IsEmpty() && ForeignEngineDir[ForeignEngineDir.Len() - 1] != TEXT('/'))
+	{
+		ForeignEngineDir += TEXT("/");
+	}
+
 	GForeignEngineDir = *ForeignEngineDir;
 	bool bInitSucceded = DatasmithGameThread::InitializeInCurrentThread(PreInitCommandArgs, bSuppressLogs);
 
