[Solidworks] Animation are now exported from solidworks using Direct Link Sync

- DatasmithFacade exposed SerializeLevelSequences method to allow syncing of animations

diff --git a/Engine/Source/Developer/Datasmith/DatasmithFacade/Private/DatasmithFacadeScene.cpp b/Engine/Source/Developer/Datasmith/DatasmithFacade/Private/DatasmithFacadeScene.cpp
--- a/Engine/Source/Developer/Datasmith/DatasmithFacade/Private/DatasmithFacadeScene.cpp
+++ b/Engine/Source/Developer/Datasmith/DatasmithFacade/Private/DatasmithFacadeScene.cpp
@@ -13,11 +13,13 @@
 #include "DatasmithFacadeAnimation.h"
 
 // Datasmith SDK.
+#include "DatasmithAnimationSerializer.h"
 #include "DatasmithExporterManager.h"
 #include "DatasmithExportOptions.h"
 #include "DatasmithSceneExporter.h"
 
 #include "Misc/Paths.h"
+#include "HAL/FileManager.h"
 
 FDatasmithFacadeScene::FDatasmithFacadeScene(
 	const TCHAR* InApplicationHostName,
@@ -543,6 +545,30 @@ bool FDatasmithFacadeScene::ExportScene(bool bCleanupUnusedElements)
 	return true;
 }
 
+void FDatasmithFacadeScene::SerializeLevelSequences()
+{
+	FDatasmithAnimationSerializer AnimSerializer;
+	int32 NumSequences = SceneRef->GetLevelSequencesCount();
+	for (int32 SequenceIndex = 0; SequenceIndex < NumSequences; ++SequenceIndex)
+	{
+		const TSharedPtr<IDatasmithLevelSequenceElement>& LevelSequence = SceneRef->GetLevelSequence(SequenceIndex);
+		if (LevelSequence.IsValid())
+		{
+			FString AnimFilePath = FPaths::Combine(SceneExporterRef->GetAssetsOutputPath(), LevelSequence->GetName()) + DATASMITH_ANIMATION_EXTENSION;
+
+			if (AnimSerializer.Serialize(LevelSequence.ToSharedRef(), *AnimFilePath))
+			{
+				TUniquePtr<FArchive> AnimArchive(IFileManager::Get().CreateFileReader(*AnimFilePath));
+				if (AnimArchive)
+				{
+					LevelSequence->SetFileHash(FMD5Hash::HashFileFromArchive(AnimArchive.Get()));
+				}
+				LevelSequence->SetFile(*AnimFilePath);
+			}
+		}
+	}
+}
+
 TSharedRef<IDatasmithScene> FDatasmithFacadeScene::GetScene() const
 {
 	return SceneRef;
diff --git a/Engine/Source/Developer/Datasmith/DatasmithFacade/Public/DatasmithFacadeScene.h b/Engine/Source/Developer/Datasmith/DatasmithFacade/Public/DatasmithFacadeScene.h
--- a/Engine/Source/Developer/Datasmith/DatasmithFacade/Public/DatasmithFacadeScene.h
+++ b/Engine/Source/Developer/Datasmith/DatasmithFacade/Public/DatasmithFacadeScene.h
@@ -286,6 +286,10 @@ public:
 		bool bCleanupUnusedElements = true
 	);
 
+	/** Save Datasmith LevelSequences into files and specify those files in each corresponding LevelSequenceElement
+	 */
+	void SerializeLevelSequences();
+
 	/**
 	 * Set the Datasmith scene's label.
 	 * This is mainly used in conjunction with DirectLink. The scene's label is used
diff --git a/Engine/Source/Programs/Enterprise/Datasmith/DatasmithSolidworksExporter/Private/Document.cs b/Engine/Source/Programs/Enterprise/Datasmith/DatasmithSolidworksExporter/Private/Document.cs
--- a/Engine/Source/Programs/Enterprise/Datasmith/DatasmithSolidworksExporter/Private/Document.cs
+++ b/Engine/Source/Programs/Enterprise/Datasmith/DatasmithSolidworksExporter/Private/Document.cs
@@ -221,7 +221,7 @@ namespace DatasmithSolidworks
 			List<FConfigurationData> Configs = ConfigurationExporter.ExportConfigurations(this);
 			bHasConfigurations = (Configs != null) && (Configs.Count != 0);
 			ExportToDatasmithScene(ConfigurationExporter, new FVariantName(ConfigManager.ActiveConfiguration));
-
+			DatasmithScene.SerializeLevelSequences();
 
 			ExportLights();
 
