DatasmithMeshExporter: Fix serializing twice the legacy mesh count due to changes in the way we Serialize that value introduce with CL 38785904

diff --git a/Engine/Source/Developer/Datasmith/DatasmithExporter/Private/DatasmithMeshExporter.cpp b/Engine/Source/Developer/Datasmith/DatasmithExporter/Private/DatasmithMeshExporter.cpp
--- a/Engine/Source/Developer/Datasmith/DatasmithExporter/Private/DatasmithMeshExporter.cpp
+++ b/Engine/Source/Developer/Datasmith/DatasmithExporter/Private/DatasmithMeshExporter.cpp
@@ -144,8 +144,6 @@ bool FDatasmithMeshExporterImpl::WriteMeshFile(const FDatasmithMeshExporterOptio
 		return false;
 	}
 
-	int32 LegacyMeshCount = 0;
-	*Archive << LegacyMeshCount; // the legacy importer expect a mesh count on the first bytes. Just in case a new file would end up parsed by the legacy code...
 	OutHash = Pack.Serialize(*Archive);
 
 	return !Archive->IsError();
